Created a new localization module for all the editor's general needs with regards to localization configuration scripts and localization pipeline configuration. Necessary minor refactoring to alter dependencies as needed.

diff --git a/Engine/Config/BaseEditor.ini b/Engine/Config/BaseEditor.ini
--- a/Engine/Config/BaseEditor.ini
+++ b/Engine/Config/BaseEditor.ini
@@ -432,7 +432,7 @@ bUseFullGCAssetClassNames=true
 
 
 
-[/Script/LocalizationDashboard.LocalizationDashboardSettings]
+[/Script/Localization.LocalizationSettings]
 +EngineTargetsSettings=(Name="Engine",Guid=33482D004789784C9DA695A682ACCA1B,TargetDependencies=,AdditionalManifestDependencies=,RequiredModuleNames=,GatherFromTextFiles=(IsEnabled=True,SearchDirectories=((Path="Source/Runtime/"),(Path="Source/Developer/"),(Path="Config/")),ExcludePathWildcards=((Pattern="Source/Developer/NoRedist/CommunityPortalServices/*"),(Pattern="Source/Developer/NoRedist/UnrealEngineLauncherServices/*"),(Pattern="Source/Developer/NoRedist/BuildPatchServices/*")),FileExtensions=((Pattern="cpp"),(Pattern="h"),(Pattern="c"),(Pattern="inl"),(Pattern="mm"),(Pattern="ini"))),GatherFromPackages=(IsEnabled=True,IncludePathWildcards=((Pattern="Content/*")),ExcludePathWildcards=((Pattern="Content/Editor/*"),(Pattern="Content/Tutorial/*"),(Pattern="Content/Developers/*"),(Pattern="Content/TestPackages/*"),(Pattern="Content/QA_Assets/*"),(Pattern="Content/Maps/Automation/*"),(Pattern="Content/EngineSounds/*")),FileExtensions=((Pattern="umap"),(Pattern="uasset"))),GatherFromMetaData=(IsEnabled=False,IncludePathWildcards=,ExcludePathWildcards=,KeySpecifications=),NativeCultureIndex=0,SupportedCulturesStatistics=((CultureName="en"),(CultureName="es"),(CultureName="ja"),(CultureName="ko"),(CultureName="pt-BR"),(CultureName="zh-CN")))
 +EngineTargetsSettings=(Name="Editor",Guid=AC8BFD2A41A2FB2893BB8EA0AF903E6D,TargetDependencies=(33482D004789784C9DA695A682ACCA1B),AdditionalManifestDependencies=,RequiredModuleNames=,GatherFromTextFiles=(IsEnabled=True,SearchDirectories=((Path="Source/Editor/")),ExcludePathWildcards=,FileExtensions=((Pattern="cpp"),(Pattern="h"),(Pattern="c"),(Pattern="inl"),(Pattern="mm"))),GatherFromPackages=(IsEnabled=True,IncludePathWildcards=((Pattern="Content/Editor/*"),(Pattern="Content/Editor*")),ExcludePathWildcards=,FileExtensions=((Pattern="umap"),(Pattern="uasset"))),GatherFromMetaData=(IsEnabled=True,IncludePathWildcards=((Pattern="Source/Editor/*"),(Pattern="Source/Runtime/*"),(Pattern="Source/Developer/*")),ExcludePathWildcards=((Pattern="Source/Developer/NoRedist/CommunityPortalServices/*"),(Pattern="Source/Developer/NoRedist/UnrealEngineLauncherServices/*"),(Pattern="Source/Developer/NoRedist/BuildPatchServices/*")),KeySpecifications=((MetaDataKey=(Name="Category"),TextNamespace="UObjectCategories",TextKeyPattern=(Pattern="{MetaDataValue}")))),NativeCultureIndex=0,SupportedCulturesStatistics=((CultureName="en"),(CultureName="es"),(CultureName="ja"),(CultureName="ko"),(CultureName="pt-BR"),(CultureName="zh-CN")))
 +EngineTargetsSettings=(Name="EditorTutorials",Guid=00F8E3AD47F0A73D50D46881C14DF28F,TargetDependencies=(33482D004789784C9DA695A682ACCA1B,AC8BFD2A41A2FB2893BB8EA0AF903E6D),AdditionalManifestDependencies=,RequiredModuleNames=("IntroTutorials"),GatherFromTextFiles=(IsEnabled=False,SearchDirectories=,ExcludePathWildcards=,FileExtensions=((Pattern="h"),(Pattern="cpp"),(Pattern="ini"))),GatherFromPackages=(IsEnabled=True,IncludePathWildcards=((Pattern="Content/Tutorial/*")),ExcludePathWildcards=,FileExtensions=((Pattern="umap"),(Pattern="uasset"))),GatherFromMetaData=(IsEnabled=False,IncludePathWildcards=,ExcludePathWildcards=,KeySpecifications=),NativeCultureIndex=0,SupportedCulturesStatistics=((CultureName="en"),(CultureName="es"),(CultureName="ja"),(CultureName="ko"),(CultureName="pt-BR"),(CultureName="zh-CN")))
diff --git a/Engine/Source/Editor/Localization/Localization.Build.cs b/Engine/Source/Editor/Localization/Localization.Build.cs
new file mode 100644
--- /dev/null
+++ b/Engine/Source/Editor/Localization/Localization.Build.cs
@@ -0,0 +1,38 @@
+// Copyright 1998-2015 Epic Games, Inc. All Rights Reserved.
+
+using UnrealBuildTool;
+
+public class Localization : ModuleRules
+{
+	public Localization(TargetInfo Target)
+	{
+		PublicDependencyModuleNames.AddRange(
+			new string[]
+			{
+				"Core",
+                "CoreUObject",
+			}
+		);
+
+		PrivateDependencyModuleNames.AddRange(
+			new string[]
+			{
+				"Engine",
+			}
+		);
+
+		PrivateIncludePaths.AddRange(
+			new string[]
+			{
+				"Editor/Localization/Private",
+			}
+		);
+
+        PublicIncludePaths.AddRange(
+			new string[]
+			{
+				"Editor/Localization/Public",
+			}
+		);
+	}
+}
\ No newline at end of file
diff --git a/Engine/Source/Editor/Localization/Private/LocalizationConfigurationScript.cpp b/Engine/Source/Editor/Localization/Private/LocalizationConfigurationScript.cpp
new file mode 100644
--- /dev/null
+++ b/Engine/Source/Editor/Localization/Private/LocalizationConfigurationScript.cpp
@@ -0,0 +1,630 @@
+﻿// Copyright 1998-2015 Epic Games, Inc. All Rights Reserved.
+
+#include "LocalizationPrivatePCH.h"
+#include "LocalizationConfigurationScript.h"
+#include "LocalizationTargetTypes.h"
+#include "LocalizationSettings.h"
+
+namespace
+{
+	FString GetConfigDir(const ULocalizationTarget* const Target)
+	{
+		return Target->IsMemberOfEngineTargetSet() ? FPaths::EngineConfigDir() : FPaths::GameConfigDir();
+	}
+
+	FString GetContentDir(const ULocalizationTarget* const Target)
+	{
+		return Target->IsMemberOfEngineTargetSet() ? FPaths::EngineContentDir() : FPaths::GameContentDir();
+	}
+}
+
+namespace LocalizationConfigurationScript
+{
+	FString MakePathRelativeForCommandletProcess(const FString& Path, const bool IsUsingProjectFile)
+	{
+		FString Result = Path;
+		const FString ProjectDir = !IsUsingProjectFile ? FPaths::EngineDir() : FPaths::GameDir();
+		if (!FPaths::MakePathRelativeTo(Result, *ProjectDir))
+		{
+			Result = FPaths::ConvertRelativePathToFull(Path);
+		}
+		return Result;
+	}
+
+	FString GetScriptDirectory(const ULocalizationTarget* const Target)
+	{
+		return GetConfigDir(Target) / TEXT("Localization");
+	}
+
+	FString GetDataDirectory(const ULocalizationTarget* const Target)
+	{
+		return GetContentDir(Target) / TEXT("Localization") / Target->Settings.Name;
+	}
+
+	TArray<FString> GetScriptPaths(const ULocalizationTarget* const Target)
+	{
+		TArray<FString> Result;
+		Result.Add(GetGatherScriptPath(Target));
+		Result.Add(GetImportScriptPath(Target));
+		Result.Add(GetExportScriptPath(Target));
+		Result.Add(GetWordCountReportScriptPath(Target));
+		return Result;
+	}
+
+	FString GetManifestPath(const ULocalizationTarget* const Target)
+	{
+		return GetDataDirectory(Target) / FString::Printf( TEXT("%s.%s"), *Target->Settings.Name, TEXT("manifest") );
+	}
+
+	FString GetArchivePath(const ULocalizationTarget* const Target, const FString& CultureName)
+	{
+		return GetDataDirectory(Target) / CultureName / FString::Printf( TEXT("%s.%s"), *Target->Settings.Name, TEXT("archive") );
+	}
+
+	FString GetDefaultPOFileName(const ULocalizationTarget* const Target)
+	{
+		return FString::Printf( TEXT("%s.%s"), *Target->Settings.Name, TEXT("po") );
+	}
+
+	FString GetDefaultPOPath(const ULocalizationTarget* const Target, const FString& CultureName)
+	{
+		return GetDataDirectory(Target) / CultureName / GetDefaultPOFileName(Target);
+	}
+
+	FString GetLocResPath(const ULocalizationTarget* const Target, const FString& CultureName)
+	{
+		return GetDataDirectory(Target) / CultureName / FString::Printf( TEXT("%s.%s"), *Target->Settings.Name, TEXT("locres") );
+	}
+
+	FString GetWordCountCSVPath(const ULocalizationTarget* const Target)
+	{
+		return GetDataDirectory(Target) / FString::Printf( TEXT("%s.%s"), *Target->Settings.Name, TEXT("csv") );
+	}
+
+	FString GetConflictReportPath(const ULocalizationTarget* const Target)
+	{
+		return GetDataDirectory(Target) / FString::Printf( TEXT("%s_Conflicts.%s"), *Target->Settings.Name, TEXT("txt") );
+	}
+
+	FLocalizationConfigurationScript GenerateGatherScript(const ULocalizationTarget* const Target)
+	{
+		FLocalizationConfigurationScript Script;
+
+		const FString ConfigDirRelativeToGameDir = MakePathRelativeForCommandletProcess(GetConfigDir(Target), !Target->IsMemberOfEngineTargetSet());
+		const FString ContentDirRelativeToGameDir = MakePathRelativeForCommandletProcess(GetContentDir(Target), !Target->IsMemberOfEngineTargetSet());
+
+		// CommonSettings
+		{
+			FConfigSection& ConfigSection = Script.CommonSettings();
+
+			const ULocalizationTargetSet* const LocalizationTargetSet = GetDefault<ULocalizationTargetSet>(ULocalizationTargetSet::StaticClass());
+			for (const FGuid& TargetDependencyGuid : Target->Settings.TargetDependencies)
+			{
+				TArray<ULocalizationTarget*> AllLocalizationTargets;
+				ULocalizationTargetSet* EngineTargetSet = ULocalizationSettings::GetEngineTargetSet();
+				if (EngineTargetSet != LocalizationTargetSet)
+				{
+					AllLocalizationTargets.Append(EngineTargetSet->TargetObjects);
+				}
+				AllLocalizationTargets.Append(LocalizationTargetSet->TargetObjects);
+
+				ULocalizationTarget* const * OtherTarget = AllLocalizationTargets.FindByPredicate([&TargetDependencyGuid](ULocalizationTarget* const InOtherTarget)->bool{return InOtherTarget->Settings.Guid == TargetDependencyGuid;});
+				if (OtherTarget)
+				{
+					ConfigSection.Add( TEXT("ManifestDependencies"), MakePathRelativeForCommandletProcess(GetManifestPath(*OtherTarget), !Target->IsMemberOfEngineTargetSet()) );
+				}
+			}
+
+			for (const FFilePath& Path : Target->Settings.AdditionalManifestDependencies)
+			{
+				ConfigSection.Add( TEXT("ManifestDependencies"), MakePathRelativeForCommandletProcess(Path.FilePath, !Target->IsMemberOfEngineTargetSet()) );
+			}
+
+			for (const FString& ModuleName : Target->Settings.RequiredModuleNames)
+			{
+				ConfigSection.Add( TEXT("ModulesToPreload"), ModuleName );
+			}
+
+			const FString SourcePath = ContentDirRelativeToGameDir / TEXT("Localization") / Target->Settings.Name;
+			ConfigSection.Add( TEXT("SourcePath"), SourcePath );
+			const FString DestinationPath = ContentDirRelativeToGameDir / TEXT("Localization") / Target->Settings.Name;
+			ConfigSection.Add( TEXT("DestinationPath"), DestinationPath );
+
+			ConfigSection.Add( TEXT("ManifestName"), FPaths::GetCleanFilename(GetManifestPath(Target)) );
+			ConfigSection.Add( TEXT("ArchiveName"), FPaths::GetCleanFilename(GetArchivePath(Target, FString())) );
+
+			if (Target->Settings.SupportedCulturesStatistics.IsValidIndex(Target->Settings.NativeCultureIndex))
+			{
+				ConfigSection.Add( TEXT("NativeCulture"), Target->Settings.SupportedCulturesStatistics[Target->Settings.NativeCultureIndex].CultureName );
+			}
+			for (const FCultureStatistics& CultureStatistics : Target->Settings.SupportedCulturesStatistics)
+			{
+				ConfigSection.Add( TEXT("CulturesToGenerate"), CultureStatistics.CultureName );
+			}
+		}
+
+		uint32 GatherTextStepIndex = 0;
+		// GatherTextFromSource
+		if (Target->Settings.GatherFromTextFiles.IsEnabled)
+		{
+			FConfigSection& ConfigSection = Script.GatherTextStep(GatherTextStepIndex++);
+
+			// CommandletClass
+			ConfigSection.Add( TEXT("CommandletClass"), TEXT("GatherTextFromSource") );
+
+			// Include Paths
+			for (const auto& IncludePath : Target->Settings.GatherFromTextFiles.SearchDirectories)
+			{
+				ConfigSection.Add( TEXT("SearchDirectoryPaths"), IncludePath.Path );
+			}
+
+			// Exclude Paths
+			ConfigSection.Add( TEXT("ExcludePathFilters"), TEXT("Config/Localization/*") );
+			for (const auto& ExcludePath : Target->Settings.GatherFromTextFiles.ExcludePathWildcards)
+			{
+				ConfigSection.Add( TEXT("ExcludePathFilters"), ExcludePath.Pattern );
+			}
+
+			// Source File Search Filters
+			for (const auto& FileExtension : Target->Settings.GatherFromTextFiles.FileExtensions)
+			{
+				ConfigSection.Add( TEXT("FileNameFilters"), FString::Printf( TEXT("*.%s"), *FileExtension.Pattern) );
+			}
+		}
+
+		// GatherTextFromAssets
+		if (Target->Settings.GatherFromPackages.IsEnabled)
+		{
+			FConfigSection& ConfigSection = Script.GatherTextStep(GatherTextStepIndex++);
+
+			// CommandletClass
+			ConfigSection.Add( TEXT("CommandletClass"), TEXT("GatherTextFromAssets") );
+
+			// Include Paths
+			for (const auto& IncludePath : Target->Settings.GatherFromPackages.IncludePathWildcards)
+			{
+				ConfigSection.Add( TEXT("IncludePathFilters"), IncludePath.Pattern );
+			}
+
+			// Exclude Paths
+			ConfigSection.Add( TEXT("ExcludePathFilters"), TEXT("Content/Localization/*") );
+			for (const auto& ExcludePath : Target->Settings.GatherFromPackages.ExcludePathWildcards)
+			{
+				ConfigSection.Add( TEXT("ExcludePathFilters"), ExcludePath.Pattern );
+			}
+
+			// Package Extensions
+			for (const auto& FileExtension : Target->Settings.GatherFromPackages.FileExtensions)
+			{
+				ConfigSection.Add( TEXT("PackageFileNameFilters"), FString::Printf( TEXT("*.%s"), *FileExtension.Pattern) );
+			}
+		}
+
+		// GatherTextFromMetadata
+		if (Target->Settings.GatherFromMetaData.IsEnabled)
+		{
+			FConfigSection& ConfigSection = Script.GatherTextStep(GatherTextStepIndex++);
+
+			// CommandletClass
+			ConfigSection.Add( TEXT("CommandletClass"), TEXT("GatherTextFromMetadata") );
+
+			// Include Paths
+			for (const auto& IncludePath : Target->Settings.GatherFromMetaData.IncludePathWildcards)
+			{
+				ConfigSection.Add( TEXT("IncludePathFilters"), IncludePath.Pattern );
+			}
+
+			// Exclude Paths
+			for (const auto& ExcludePath : Target->Settings.GatherFromMetaData.ExcludePathWildcards)
+			{
+				ConfigSection.Add( TEXT("ExcludePathFilters"), ExcludePath.Pattern );
+			}
+
+			// Package Extensions
+			for (const FMetaDataKeyGatherSpecification& Specification : Target->Settings.GatherFromMetaData.KeySpecifications)
+			{
+				ConfigSection.Add( TEXT("InputKeys"), Specification.MetaDataKey.Name );
+				ConfigSection.Add( TEXT("OutputNamespaces"), Specification.TextNamespace );
+				ConfigSection.Add( TEXT("OutputKeys"), FString::Printf(TEXT("\"%s\""), *Specification.TextKeyPattern.Pattern) );
+			}
+		}
+
+		// GenerateGatherManifest
+		{
+			FConfigSection& ConfigSection = Script.GatherTextStep(GatherTextStepIndex++);
+
+			// CommandletClass
+			ConfigSection.Add( TEXT("CommandletClass"), TEXT("GenerateGatherManifest") );
+		}
+
+		// GenerateGatherArchive
+		{
+			FConfigSection& ConfigSection = Script.GatherTextStep(GatherTextStepIndex++);
+
+			// CommandletClass
+			ConfigSection.Add( TEXT("CommandletClass"), TEXT("GenerateGatherArchive") );
+		}
+
+		// GenerateTextLocalizationReport
+		{
+			FConfigSection& ConfigSection = Script.GatherTextStep(GatherTextStepIndex++);
+
+			// CommandletClass
+			ConfigSection.Add( TEXT("CommandletClass"), TEXT("GenerateTextLocalizationReport") );
+
+			ConfigSection.Add( TEXT("bWordCountReport"), TEXT("true") );
+			ConfigSection.Add( TEXT("WordCountReportName"), FPaths::GetCleanFilename( GetWordCountCSVPath(Target) ) );
+
+			ConfigSection.Add( TEXT("bConflictReport"), TEXT("true") );
+			ConfigSection.Add( TEXT("ConflictReportName"), FPaths::GetCleanFilename( GetConflictReportPath(Target) ) );
+		}
+
+		Script.Dirty = true;
+
+		return Script;
+	}
+
+	FString GetGatherScriptPath(const ULocalizationTarget* const Target)
+	{
+		return GetScriptDirectory(Target) / FString::Printf( TEXT("%s_Gather.%s"), *(Target->Settings.Name), TEXT("ini") );
+	}
+
+	FLocalizationConfigurationScript GenerateImportScript(const ULocalizationTarget* const Target, const TOptional<FString> CultureName, const TOptional<FString> OutputPathOverride)
+	{
+		FLocalizationConfigurationScript Script;
+
+		const FString ContentDirRelativeToGameDir = MakePathRelativeForCommandletProcess(GetContentDir(Target), !Target->IsMemberOfEngineTargetSet());
+
+		// GatherTextStep0 - InternationalizationExport
+		{
+			FConfigSection& ConfigSection = Script.GatherTextStep(0);
+
+			// CommandletClass
+			ConfigSection.Add( TEXT("CommandletClass"), TEXT("InternationalizationExport") );
+
+			ConfigSection.Add( TEXT("bImportLoc"), TEXT("true") );
+
+			FString SourcePath;
+			// Overriding output path changes the source directory for the PO file.
+			if (OutputPathOverride.IsSet())
+			{
+				// The output path for a specific culture is a file path.
+				if (CultureName.IsSet())
+				{
+					SourcePath = MakePathRelativeForCommandletProcess( FPaths::GetPath(OutputPathOverride.GetValue()), !Target->IsMemberOfEngineTargetSet() );
+				}
+				// Otherwise, it is a directory path.
+				else
+				{
+					SourcePath = MakePathRelativeForCommandletProcess( OutputPathOverride.GetValue(), !Target->IsMemberOfEngineTargetSet() );
+				}
+			}
+			// Use the default PO file's directory path.
+			else
+			{
+				SourcePath = ContentDirRelativeToGameDir / TEXT("Localization") / Target->Settings.Name;
+			}
+			ConfigSection.Add( TEXT("SourcePath"), SourcePath );
+
+			const FString DestinationPath = ContentDirRelativeToGameDir / TEXT("Localization") / Target->Settings.Name;
+			ConfigSection.Add( TEXT("DestinationPath"), DestinationPath );
+
+			if (Target->Settings.SupportedCulturesStatistics.IsValidIndex(Target->Settings.NativeCultureIndex))
+			{
+				ConfigSection.Add( TEXT("NativeCulture"), Target->Settings.SupportedCulturesStatistics[Target->Settings.NativeCultureIndex].CultureName );
+			}
+
+			const auto& AddCultureToGenerate = [&](const int32 Index)
+			{
+				ConfigSection.Add( TEXT("CulturesToGenerate"), Target->Settings.SupportedCulturesStatistics[Index].CultureName );
+			};
+
+			// Export for a specific culture.
+			if (CultureName.IsSet())
+			{
+				ConfigSection.Add( TEXT("CulturesToGenerate"), CultureName.GetValue() );
+			}
+			// Export for all cultures.
+			else
+			{
+				for (const FCultureStatistics& CultureStatistics : Target->Settings.SupportedCulturesStatistics)
+				{
+					ConfigSection.Add( TEXT("CulturesToGenerate"), CultureStatistics.CultureName );
+				}
+			}
+
+			// Do not use culture subdirectories if importing a single culture to a specific directory.
+			if (CultureName.IsSet() && OutputPathOverride.IsSet())
+			{
+				ConfigSection.Add( TEXT("bUseCultureDirectory"), "false" );
+			}
+
+			ConfigSection.Add( TEXT("ManifestName"), FPaths::GetCleanFilename(GetManifestPath(Target)) );
+			ConfigSection.Add( TEXT("ArchiveName"), FPaths::GetCleanFilename(GetArchivePath(Target, FString())) );
+
+			FString POFileName;
+			// The output path for a specific culture is a file path.
+			if (CultureName.IsSet() && OutputPathOverride.IsSet())
+			{
+				POFileName =  FPaths::GetCleanFilename( OutputPathOverride.GetValue() );
+			}
+			// Use the default PO file's name.
+			else
+			{
+				POFileName = FPaths::GetCleanFilename( GetDefaultPOFileName( Target ) );
+			}
+			ConfigSection.Add( TEXT("PortableObjectName"), POFileName );
+		}
+
+		Script.Dirty = true;
+
+		return Script;
+	}
+
+	FString GetImportScriptPath(const ULocalizationTarget* const Target, const TOptional<FString> CultureName)
+	{
+		const FString ConfigFileDirectory = GetScriptDirectory(Target);
+		FString ConfigFilePath;
+		if (CultureName.IsSet())
+		{
+			ConfigFilePath = ConfigFileDirectory / FString::Printf( TEXT("%s_Import_%s.%s"), *Target->Settings.Name, *CultureName.GetValue(), TEXT("ini") );
+		}
+		else
+		{
+			ConfigFilePath = ConfigFileDirectory / FString::Printf( TEXT("%s_Import.%s"), *Target->Settings.Name, TEXT("ini") );
+		}
+		return ConfigFilePath;
+	}
+
+	FLocalizationConfigurationScript GenerateExportScript(const ULocalizationTarget* const Target, const TOptional<FString> CultureName, const TOptional<FString> OutputPathOverride)
+	{
+		FLocalizationConfigurationScript Script;
+
+		const FString ContentDirRelativeToGameDir = MakePathRelativeForCommandletProcess(GetContentDir(Target), !Target->IsMemberOfEngineTargetSet());
+
+		// GatherTextStep0 - InternationalizationExport
+		{
+			FConfigSection& ConfigSection = Script.GatherTextStep(0);
+
+			// CommandletClass
+			ConfigSection.Add( TEXT("CommandletClass"), TEXT("InternationalizationExport") );
+
+			ConfigSection.Add( TEXT("bExportLoc"), TEXT("true") );
+
+			const FString SourcePath = ContentDirRelativeToGameDir / TEXT("Localization") / Target->Settings.Name;
+			ConfigSection.Add( TEXT("SourcePath"), SourcePath );
+
+			FString DestinationPath;
+			// Overriding output path changes the destination directory for the PO file.
+			if (OutputPathOverride.IsSet())
+			{
+				// The output path for a specific culture is a file path.
+				if (CultureName.IsSet())
+				{
+					DestinationPath = MakePathRelativeForCommandletProcess( FPaths::GetPath(OutputPathOverride.GetValue()), !Target->IsMemberOfEngineTargetSet() );
+				}
+				// Otherwise, it is a directory path.
+				else
+				{
+					DestinationPath = MakePathRelativeForCommandletProcess( OutputPathOverride.GetValue(), !Target->IsMemberOfEngineTargetSet() );
+				}
+			}
+			// Use the default PO file's directory path.
+			else
+			{
+				DestinationPath = ContentDirRelativeToGameDir / TEXT("Localization") / Target->Settings.Name;
+			}
+			ConfigSection.Add( TEXT("DestinationPath"), DestinationPath );
+
+			if (Target->Settings.SupportedCulturesStatistics.IsValidIndex(Target->Settings.NativeCultureIndex))
+			{
+				ConfigSection.Add( TEXT("NativeCulture"), Target->Settings.SupportedCulturesStatistics[Target->Settings.NativeCultureIndex].CultureName );
+			}
+
+			const auto& AddCultureToGenerate = [&](const int32 Index)
+			{
+				ConfigSection.Add( TEXT("CulturesToGenerate"), Target->Settings.SupportedCulturesStatistics[Index].CultureName );
+			};
+
+			// Export for a specific culture.
+			if (CultureName.IsSet())
+			{
+				const int32 CultureIndex = Target->Settings.SupportedCulturesStatistics.IndexOfByPredicate([CultureName](const FCultureStatistics& Culture) { return Culture.CultureName == CultureName.GetValue(); });
+				AddCultureToGenerate(CultureIndex);
+			}
+			// Export for all cultures.
+			else
+			{
+				for (int32 CultureIndex = 0; CultureIndex < Target->Settings.SupportedCulturesStatistics.Num(); ++CultureIndex)
+				{
+					AddCultureToGenerate(CultureIndex);
+				}
+			}
+
+			// Do not use culture subdirectories if exporting a single culture to a specific directory.
+			if (CultureName.IsSet() && OutputPathOverride.IsSet())
+			{
+				ConfigSection.Add( TEXT("bUseCultureDirectory"), "false" );
+			}
+
+
+			ConfigSection.Add( TEXT("ManifestName"), FPaths::GetCleanFilename(GetManifestPath(Target)) );
+			ConfigSection.Add( TEXT("ArchiveName"), FPaths::GetCleanFilename(GetArchivePath(Target, FString())) );
+			FString POFileName;
+			// The output path for a specific culture is a file path.
+			if (CultureName.IsSet() && OutputPathOverride.IsSet())
+			{
+				POFileName =  FPaths::GetCleanFilename( OutputPathOverride.GetValue() );
+			}
+			// Use the default PO file's name.
+			else
+			{
+				POFileName = FPaths::GetCleanFilename( GetDefaultPOPath( Target, CultureName.Get( TEXT("") ) ) );
+			}
+			ConfigSection.Add( TEXT("PortableObjectName"), POFileName );
+		}
+
+		Script.Dirty = true;
+
+		return Script;
+	}
+
+	FString GetExportScriptPath(const ULocalizationTarget* const Target, const TOptional<FString> CultureName)
+	{
+		const FString ConfigFileDirectory = GetScriptDirectory(Target);
+		FString ConfigFilePath;
+		if (CultureName.IsSet())
+		{
+			ConfigFilePath = ConfigFileDirectory / FString::Printf( TEXT("%s_Export_%s.%s"), *Target->Settings.Name, *CultureName.GetValue(), TEXT("ini") );
+		}
+		else
+		{
+			ConfigFilePath = ConfigFileDirectory / FString::Printf( TEXT("%s_Export.%s"), *Target->Settings.Name, TEXT("ini") );
+		}
+		return ConfigFilePath;
+	}
+
+	FLocalizationConfigurationScript GenerateWordCountReportScript(const ULocalizationTarget* const Target)
+	{
+		FLocalizationConfigurationScript Script;
+
+		const FString ContentDirRelativeToGameDir = MakePathRelativeForCommandletProcess(GetContentDir(Target), !Target->IsMemberOfEngineTargetSet());
+
+		// GatherTextStep0 - GenerateTextLocalizationReport
+		{
+			FConfigSection& ConfigSection = Script.GatherTextStep(0);
+
+			// CommandletClass
+			ConfigSection.Add( TEXT("CommandletClass"), TEXT("GenerateTextLocalizationReport") );
+
+			const FString SourcePath = ContentDirRelativeToGameDir / TEXT("Localization") / Target->Settings.Name;
+			ConfigSection.Add( TEXT("SourcePath"), SourcePath );
+			const FString DestinationPath = ContentDirRelativeToGameDir / TEXT("Localization") / Target->Settings.Name;
+			ConfigSection.Add( TEXT("DestinationPath"), DestinationPath );
+
+			ConfigSection.Add( TEXT("ManifestName"), FString::Printf( TEXT("%s.%s"), *Target->Settings.Name, TEXT("manifest") ) );
+
+			for (const FCultureStatistics& CultureStatistics : Target->Settings.SupportedCulturesStatistics)
+			{
+				ConfigSection.Add( TEXT("CulturesToGenerate"), CultureStatistics.CultureName );
+			}
+
+			ConfigSection.Add( TEXT("bWordCountReport"), TEXT("true") );
+
+			ConfigSection.Add( TEXT("WordCountReportName"), FPaths::GetCleanFilename( GetWordCountCSVPath(Target) ) );
+		}
+
+		Script.Dirty = true;
+
+		return Script;
+	}
+
+	FString GetWordCountReportScriptPath(const ULocalizationTarget* const Target)
+	{
+		return GetScriptDirectory(Target) / FString::Printf( TEXT("%s_GenerateReports.%s"), *(Target->Settings.Name), TEXT("ini") );
+	}
+
+	FLocalizationConfigurationScript GenerateCompileScript(const ULocalizationTarget* const Target, const TOptional<FString> CultureName)
+	{
+		FLocalizationConfigurationScript Script;
+
+		const FString ContentDirRelativeToGameDir = MakePathRelativeForCommandletProcess(GetContentDir(Target), !Target->IsMemberOfEngineTargetSet());
+
+		// GatherTextStep0 - GenerateTextLocalizationResource
+		{
+			FConfigSection& ConfigSection = Script.GatherTextStep(0);
+
+			// CommandletClass
+			ConfigSection.Add( TEXT("CommandletClass"), TEXT("GenerateTextLocalizationResource") );
+
+			const FString SourcePath = ContentDirRelativeToGameDir / TEXT("Localization") / Target->Settings.Name;
+			ConfigSection.Add( TEXT("SourcePath"), SourcePath );
+			const FString DestinationPath = ContentDirRelativeToGameDir / TEXT("Localization") / Target->Settings.Name;
+			ConfigSection.Add( TEXT("DestinationPath"), DestinationPath );
+
+			ConfigSection.Add( TEXT("ManifestName"), FString::Printf( TEXT("%s.%s"), *Target->Settings.Name, TEXT("manifest") ) );
+			ConfigSection.Add( TEXT("ResourceName"), FString::Printf( TEXT("%s.%s"), *Target->Settings.Name, TEXT("locres") ) );
+
+			if (Target->Settings.SupportedCulturesStatistics.IsValidIndex(Target->Settings.NativeCultureIndex))
+			{
+				ConfigSection.Add( TEXT("NativeCulture"), Target->Settings.SupportedCulturesStatistics[Target->Settings.NativeCultureIndex].CultureName );
+			}
+
+			const auto& AddCultureToGenerate = [&](const int32 Index)
+			{
+				ConfigSection.Add( TEXT("CulturesToGenerate"), Target->Settings.SupportedCulturesStatistics[Index].CultureName );
+			};
+
+			// Compile a specific culture.
+			if (CultureName.IsSet())
+			{
+				const int32 CultureIndex = Target->Settings.SupportedCulturesStatistics.IndexOfByPredicate([CultureName](const FCultureStatistics& Culture) { return Culture.CultureName == CultureName.GetValue(); });
+				AddCultureToGenerate(CultureIndex);
+			}
+			// Compile all cultures.
+			else
+			{
+				for (int32 CultureIndex = 0; CultureIndex < Target->Settings.SupportedCulturesStatistics.Num(); ++CultureIndex)
+			{
+					AddCultureToGenerate(CultureIndex);
+				}
+			}
+		}
+
+		Script.Dirty = true;
+
+		return Script;
+	}
+
+	FString GetCompileScriptPath(const ULocalizationTarget* const Target, const TOptional<FString> CultureName)
+	{
+		const FString ConfigFileDirectory = GetScriptDirectory(Target);
+		FString ConfigFilePath;
+		if (CultureName.IsSet())
+		{
+			ConfigFilePath = ConfigFileDirectory / FString::Printf( TEXT("%s_Compile_%s.%s"), *Target->Settings.Name, *CultureName.GetValue(), TEXT("ini") );
+		}
+		else
+		{
+			ConfigFilePath = ConfigFileDirectory / FString::Printf( TEXT("%s_Compile.%s"), *Target->Settings.Name, TEXT("ini") );
+		}
+		return ConfigFilePath;
+	}
+
+	FLocalizationConfigurationScript GenerateRegenerateResourcesScript(const ULocalizationTarget* const Target)
+	{
+		FLocalizationConfigurationScript Script;
+
+		const FString ContentDirRelativeToGameDir = MakePathRelativeForCommandletProcess(GetContentDir(Target), !Target->IsMemberOfEngineTargetSet());
+
+		// RegenerateResources
+		{
+			FConfigSection& ConfigSection = Script.FindOrAdd("RegenerateResources");;
+
+			if (Target->Settings.SupportedCulturesStatistics.IsValidIndex(Target->Settings.NativeCultureIndex))
+			{
+				ConfigSection.Add(TEXT("NativeCulture"), Target->Settings.SupportedCulturesStatistics[Target->Settings.NativeCultureIndex].CultureName);
+			}
+
+			const FString SourcePath = ContentDirRelativeToGameDir / TEXT("Localization") / Target->Settings.Name;
+			ConfigSection.Add(TEXT("SourcePath"), SourcePath);
+			const FString DestinationPath = ContentDirRelativeToGameDir / TEXT("Localization") / Target->Settings.Name;
+			ConfigSection.Add(TEXT("DestinationPath"), DestinationPath);
+
+			ConfigSection.Add(TEXT("ManifestName"), FString::Printf(TEXT("%s.%s"), *Target->Settings.Name, TEXT("manifest")));
+
+			ConfigSection.Add(TEXT("ResourceName"), FString::Printf(TEXT("%s.%s"), *Target->Settings.Name, TEXT("locres")));
+
+		}
+
+		Script.Dirty = true;
+
+		return Script;
+	}
+
+	FString GetRegenerateResourcesScriptPath(const ULocalizationTarget* const Target)
+	{
+		return GetScriptDirectory(Target) / FString::Printf(TEXT("Regenerate%s.%s"), *(Target->Settings.Name), TEXT("ini"));
+	}
+
+}
\ No newline at end of file
diff --git a/Engine/Source/Editor/Localization/Private/LocalizationModule.cpp b/Engine/Source/Editor/Localization/Private/LocalizationModule.cpp
new file mode 100644
--- /dev/null
+++ b/Engine/Source/Editor/Localization/Private/LocalizationModule.cpp
@@ -0,0 +1,37 @@
+﻿// Copyright 1998-2015 Epic Games, Inc. All Rights Reserved.
+
+#include "LocalizationPrivatePCH.h"
+#include "LocalizationModule.h"
+#include "LocalizationTargetTypes.h"
+#include "LocalizationSettings.h"
+
+ULocalizationTarget* FLocalizationModule::GetLocalizationTargetByName(FString TargetName, bool bIsEngineTarget)
+{
+	if (bIsEngineTarget)
+	{
+		ULocalizationTargetSet* EngineTargetSet = ULocalizationSettings::GetEngineTargetSet();
+		for (ULocalizationTarget* Target : EngineTargetSet->TargetObjects)
+		{
+			if (Target->Settings.Name == TargetName)
+			{
+				return Target;
+			}
+		}
+	}
+	else
+	{
+		ULocalizationTargetSet* EngineTargetSet = ULocalizationSettings::GetEngineTargetSet();
+		for (ULocalizationTarget* Target : EngineTargetSet->TargetObjects)
+		{
+			if (Target->Settings.Name == TargetName)
+			{
+				return Target;
+			}
+		}
+	}
+
+	return nullptr;
+}
+
+
+IMPLEMENT_MODULE(FLocalizationModule, Localization);
\ No newline at end of file
diff --git a/Engine/Source/Editor/Localization/Private/LocalizationPrivatePCH.h b/Engine/Source/Editor/Localization/Private/LocalizationPrivatePCH.h
new file mode 100644
--- /dev/null
+++ b/Engine/Source/Editor/Localization/Private/LocalizationPrivatePCH.h
@@ -0,0 +1,10 @@
+// Copyright 1998-2015 Epic Games, Inc. All Rights Reserved.
+
+#pragma once
+
+
+/* Dependencies
+ *****************************************************************************/
+
+#include "CoreUObject.h"
+#include "UnrealEd.h"
\ No newline at end of file
diff --git a/Engine/Source/Editor/Localization/Private/LocalizationSettings.cpp b/Engine/Source/Editor/Localization/Private/LocalizationSettings.cpp
new file mode 100644
--- /dev/null
+++ b/Engine/Source/Editor/Localization/Private/LocalizationSettings.cpp
@@ -0,0 +1,89 @@
+// Copyright 1998-2015 Epic Games, Inc. All Rights Reserved.
+
+#include "LocalizationPrivatePCH.h"
+#include "LocalizationSettings.h"
+#include "LocalizationTargetTypes.h"
+#include "ISourceControlModule.h"
+
+ULocalizationSettings::ULocalizationSettings(const FObjectInitializer& ObjectInitializer)
+	: Super(ObjectInitializer)
+	, EngineTargetSet(ObjectInitializer.CreateDefaultSubobject<ULocalizationTargetSet>(this, TEXT("EngineLocalizationTargetSet")))
+	, GameTargetSet(ObjectInitializer.CreateDefaultSubobject<ULocalizationTargetSet>(this, TEXT("ProjectLocalizationTargetSet")))
+{	
+}
+
+#if WITH_EDITOR
+void ULocalizationSettings::PostInitProperties()
+{
+	Super::PostInitProperties();
+
+	// Create and initialize objects for details model from backing config properties.
+	if (EngineTargetSet)
+	{
+		EngineTargetSet->TargetObjects.Empty(EngineTargetsSettings.Num());
+		for (const auto& TargetSettings : EngineTargetsSettings)
+		{
+			ULocalizationTarget* const TargetObject = NewObject<ULocalizationTarget>(EngineTargetSet);
+			TargetObject->Settings = TargetSettings;
+			TargetObject->UpdateStatusFromConflictReport();
+			TargetObject->UpdateWordCountsFromCSV();
+			EngineTargetSet->TargetObjects.Add(TargetObject);
+		}
+	}
+
+	// Create and initialize objects for details model from backing config properties.
+	if (GameTargetSet)
+	{
+		GameTargetSet->TargetObjects.Empty(EngineTargetsSettings.Num());
+		for (const auto& TargetSettings : GameTargetsSettings)
+		{
+			ULocalizationTarget* const TargetObject = NewObject<ULocalizationTarget>(GameTargetSet);
+			TargetObject->Settings = TargetSettings;
+			TargetObject->UpdateStatusFromConflictReport();
+			TargetObject->UpdateWordCountsFromCSV();
+			GameTargetSet->TargetObjects.Add(TargetObject);
+		}
+	}
+}
+
+void ULocalizationSettings::PostEditChangeProperty(FPropertyChangedEvent& PropertyChangedEvent)
+{
+	Super::PostEditChangeProperty(PropertyChangedEvent);
+
+	// Copy settings back.
+	if (EngineTargetSet)
+	{
+		EngineTargetsSettings.Empty(EngineTargetSet->TargetObjects.Num());
+		for (const auto& TargetObject : EngineTargetSet->TargetObjects)
+		{
+			EngineTargetsSettings.Add(TargetObject ? TargetObject->Settings : FLocalizationTargetSettings());
+		}
+	}
+
+	// Copy settings back.
+	if (GameTargetSet)
+	{
+		GameTargetsSettings.Empty(GameTargetSet->TargetObjects.Num());
+		for (const auto& TargetObject : GameTargetSet->TargetObjects)
+		{
+			GameTargetsSettings.Add(TargetObject ? TargetObject->Settings : FLocalizationTargetSettings());
+		}
+	}
+
+	UpdateDefaultConfigFile();
+}
+#endif
+
+ULocalizationTargetSet* ULocalizationSettings::GetEngineTargetSet()
+{
+	ULocalizationSettings* LocalizationSettings = GetMutableDefault<ULocalizationSettings>();
+	check(LocalizationSettings);
+	return LocalizationSettings->EngineTargetSet;
+}
+
+ULocalizationTargetSet* ULocalizationSettings::GetGameTargetSet()
+{
+	ULocalizationSettings* LocalizationSettings = GetMutableDefault<ULocalizationSettings>();
+	check(LocalizationSettings);
+	return LocalizationSettings->GameTargetSet;
+}
\ No newline at end of file
diff --git a/Engine/Source/Editor/Localization/Private/LocalizationTargetTypes.cpp b/Engine/Source/Editor/Localization/Private/LocalizationTargetTypes.cpp
new file mode 100644
--- /dev/null
+++ b/Engine/Source/Editor/Localization/Private/LocalizationTargetTypes.cpp
@@ -0,0 +1,666 @@
+﻿// Copyright 1998-2015 Epic Games, Inc. All Rights Reserved.
+
+#include "LocalizationPrivatePCH.h"
+#include "LocalizationTargetTypes.h"
+#include "CsvParser.h"
+#include "LocalizationConfigurationScript.h"
+#include "LocalizationSettings.h"
+
+#define LOCTEXT_NAMESPACE "LocalizationTargetTypes"
+
+bool FGatherTextSearchDirectory::Validate(const FString& RootDirectory, FText& OutError) const
+{
+	if (Path.IsEmpty())
+	{
+		OutError = LOCTEXT("SearchDirectoryEmptyError", "Search directory not specified. Use \".\" to specify the root directory.");
+		return false;
+	}
+
+	FText InvalidPathReason;
+	if (!FPaths::ValidatePath(Path, &InvalidPathReason))
+	{
+		OutError = InvalidPathReason;
+		return false;
+	}
+
+	if (!FPaths::DirectoryExists(FPaths::Combine(*RootDirectory, *Path)))
+	{
+		OutError = LOCTEXT("SearchDirectoryNonExistentError", "Search directory does not exist.");
+		return false;
+	}
+
+	return true;
+}
+
+bool FGatherTextIncludePath::Validate(const FString& RootDirectory, FText& OutError) const
+{
+	if (Pattern.IsEmpty())
+	{
+		OutError = LOCTEXT("IncludePathEmptyError", "Include path not specified. Use \".\" to specify the root directory.");
+		return false;
+	}
+
+	if (!Pattern.Contains(TEXT("*")))
+	{
+		OutError = LOCTEXT("IncludePathNoWildcardError", "Include path does not specify a wild card (\"*\"). Append \"*\" or only the file at the exact specified directory will be gathered from.");
+		return false;
+	}
+
+	return true;
+}
+
+bool FGatherTextExcludePath::Validate(FText& OutError) const
+{
+	if (Pattern.IsEmpty())
+	{
+		OutError = LOCTEXT("ExcludePathEmptyError", "Exclude path not specified. Use \".\" to specify the root directory.");
+		return false;
+	}
+
+	if (!Pattern.Contains(TEXT("*")))
+	{
+		OutError = LOCTEXT("ExcludePathNoWildcardError", "Exclude path does not specify a wild card (\"*\"). Append \"*\" or only the file at the exact specified directory will be excluded.");
+		return false;
+	}
+
+	return true;
+}
+
+bool FGatherTextFileExtension::Validate(FText& OutError) const
+{
+	if (Pattern.IsEmpty())
+	{
+		OutError = LOCTEXT("FileExtensionEmptyError", "File extension not specified.");
+		return false;
+	}
+
+	return true;
+}
+
+const TArray<FGatherTextFileExtension> FGatherTextFromTextFilesConfiguration::DefaultTextFileExtensions = []()
+{
+	TArray<FGatherTextFileExtension> Result;
+	Result.Add(FGatherTextFileExtension{TEXT("h")});
+	Result.Add(FGatherTextFileExtension{TEXT("cpp")});
+	Result.Add(FGatherTextFileExtension{TEXT("ini")});
+	return Result;
+}();
+
+const TArray<FGatherTextFileExtension> FGatherTextFromPackagesConfiguration::DefaultPackageFileExtensions = []()
+{
+	TArray<FGatherTextFileExtension> Result;
+	Result.Add(FGatherTextFileExtension{TEXT("umap")});
+	Result.Add(FGatherTextFileExtension{TEXT("uasset")});
+	return Result;
+}();
+
+bool FGatherTextFromTextFilesConfiguration::Validate(const FString& RootDirectory, FText& OutError) const
+{
+	const auto& InvalidSearchDirectoryPredicate = [RootDirectory](const FGatherTextSearchDirectory& Element) -> bool
+	{
+		FText ElementError;
+		return !(Element.Validate(RootDirectory, ElementError));
+	};
+	const auto& InvalidExcludePathPredicate = [RootDirectory](const FGatherTextExcludePath& Element) -> bool
+	{
+		FText ElementError;
+		return !(Element.Validate(ElementError));
+	};
+	const auto& InvalidFileExtensionPredicate = [RootDirectory](const FGatherTextFileExtension& Element) -> bool
+	{
+		FText ElementError;
+		return !(Element.Validate(ElementError));
+	};
+
+	if (SearchDirectories.Num() > 0 && !SearchDirectories.ContainsByPredicate(InvalidSearchDirectoryPredicate) &&
+		!ExcludePathWildcards.ContainsByPredicate(InvalidExcludePathPredicate) &&
+		FileExtensions.Num() > 0 && !FileExtensions.ContainsByPredicate(InvalidFileExtensionPredicate))
+	{
+		return true;
+	}
+
+	OutError = LOCTEXT("InvalidGatherTextFromFilesConfigurationError", "Must have at least one search directory, one file extension, and no invalid settings.");
+	return false;
+}
+
+bool FGatherTextFromPackagesConfiguration::Validate(const FString& RootDirectory, FText& OutError) const
+{
+	const auto& InvalidIncludePathPredicate = [RootDirectory](const FGatherTextIncludePath& Element) -> bool
+	{
+		FText ElementError;
+		return !(Element.Validate(RootDirectory, ElementError));
+	};
+	const auto& InvalidExcludePathPredicate = [RootDirectory](const FGatherTextExcludePath& Element) -> bool
+	{
+		FText ElementError;
+		return !(Element.Validate(ElementError));
+	};
+	const auto& InvalidFileExtensionPredicate = [RootDirectory](const FGatherTextFileExtension& Element) -> bool
+	{
+		FText ElementError;
+		return !(Element.Validate(ElementError));
+	};
+
+	if (IncludePathWildcards.Num() > 0 && !IncludePathWildcards.ContainsByPredicate(InvalidIncludePathPredicate) &&
+		!ExcludePathWildcards.ContainsByPredicate(InvalidExcludePathPredicate) &&
+		FileExtensions.Num() > 0 && !FileExtensions.ContainsByPredicate(InvalidFileExtensionPredicate))
+	{
+		return true;
+	}
+
+	OutError = LOCTEXT("InvalidGatherTextFromPackagesConfigurationError", "Must have at least one include path, one file extension, and no invalid settings.");
+	return false;
+}
+
+bool FMetaDataTextKeyPattern::Validate(FText& OutError) const
+{
+	for (const auto& PossiblePlaceHolder : GetPossiblePlaceHolders())
+	{
+		if (Pattern.Contains(PossiblePlaceHolder))
+		{
+			return true;
+		}
+	}
+
+	OutError = LOCTEXT("NoMetaDataLocalizationKeyPlaceHolderError", "No place holders used. All generated keys will conflict!");
+	return false;
+}
+
+const TArray<FString>& FMetaDataTextKeyPattern::GetPossiblePlaceHolders()
+{
+	static const TArray<FString> PossiblePlaceHolders = []() -> TArray<FString>
+	{
+		TArray<FString> Result;
+		Result.Add(TEXT("{FieldPath}"));
+		Result.Add(TEXT("{MetaDataValue}"));
+		return Result;
+	}();
+
+	return PossiblePlaceHolders;
+}
+
+bool FMetaDataKeyName::Validate(FText& OutError) const
+{
+	if (Name.IsEmpty())
+	{
+		OutError = LOCTEXT("MetaDataKeyNameEmptyError", "Meta data key not specified.");
+		return false;
+	}
+
+	return true;
+}
+
+bool FMetaDataKeyGatherSpecification::Validate(FText& OutError) const
+{
+	return MetaDataKey.Validate(OutError) && TextKeyPattern.Validate(OutError);
+}
+
+bool FGatherTextFromMetaDataConfiguration::Validate(const FString& RootDirectory, FText& OutError) const
+{
+	const auto& InvalidIncludePathPredicate = [RootDirectory](const FGatherTextIncludePath& Element) -> bool
+	{
+		FText ElementError;
+		return !(Element.Validate(RootDirectory, ElementError));
+	};
+	const auto& InvalidExcludePathPredicate = [RootDirectory](const FGatherTextExcludePath& Element) -> bool
+	{
+		FText ElementError;
+		return !(Element.Validate(ElementError));
+	};
+	const auto& InvalidMetaDataKeySpecificationPredicate = [RootDirectory](const FMetaDataKeyGatherSpecification& Element) -> bool
+	{
+		FText ElementError;
+		return !(Element.Validate(ElementError));
+	};
+
+	if (IncludePathWildcards.Num() > 0 && !IncludePathWildcards.ContainsByPredicate(InvalidIncludePathPredicate) &&
+		!ExcludePathWildcards.ContainsByPredicate(InvalidExcludePathPredicate) &&
+		KeySpecifications.Num() > 0 && !KeySpecifications.ContainsByPredicate(InvalidMetaDataKeySpecificationPredicate))
+	{
+		return true;
+	}
+
+	OutError = LOCTEXT("InvalidGatherTextFromMetadataConfigurationError", "Must have at least one include path, one meta data key specification, and no invalid settings.");
+	return false;
+}
+
+namespace
+{
+	class FWordCountCSVParser
+	{
+		typedef TArray<const TCHAR*> FRow;
+
+	public:
+		static bool Execute(FLocalizationTargetSettings& Target, const FString& CSVFilePath)
+		{
+			// Load CSV file to string.
+			FString CSVString;
+			if (!FFileHelper::LoadFileToString(CSVString, *CSVFilePath))
+			{
+				return false;
+			}
+
+			// Parse CSV file string.
+			FCsvParser CSVParser(CSVString);
+			const FCsvParser::FRows& CSVRows = CSVParser.GetRows();
+
+			// Initialize.
+			FWordCountCSVParser Instance(Target);
+
+			// Parse header row.
+			if ( !(CSVRows.Num() > 0) ||!Instance.ParseHeaderRow(CSVRows[0]) )
+			{
+				return false;
+			}
+
+			// Parse latest word count row.
+			TMap<FString, uint32> WordCounts;
+			if ( !Instance.ParseWordCountRow(CSVRows.Top(), WordCounts) )
+			{
+				return false;
+			}
+
+			for (FCultureStatistics& SupportedCultureStatistics : Target.SupportedCulturesStatistics)
+			{
+				uint32* SupportedWordCount = WordCounts.Find(SupportedCultureStatistics.CultureName);
+				if (SupportedWordCount)
+				{
+					SupportedCultureStatistics.WordCount = *SupportedWordCount;
+				}
+			}
+
+			return true;
+		}
+
+	private:
+		FWordCountCSVParser(FLocalizationTargetSettings& InTarget)
+			: Target(InTarget)
+		{
+			MandatoryColumnNames.Add( TEXT("Date/Time") );
+			MandatoryColumnNames.Add( TEXT("Word Count") );
+		}
+
+		bool ParseHeaderRow(const FRow& HeaderRow)
+		{
+			// Validate header row.
+			const bool IsHeaderRowValid = [&]() -> bool
+			{
+				if (HeaderRow.Num() < MandatoryColumnNames.Num())
+				{
+					// Not enough columns present to represent the mandatory columns.
+					return false;
+				}
+				for (int32 i = 0; i < MandatoryColumnNames.Num(); ++i)
+				{
+					if (MandatoryColumnNames[i] != HeaderRow[i])
+					{
+						return false;
+					}
+				}
+
+				return true;
+			}();
+
+			if (!IsHeaderRowValid)
+			{
+				return false;
+			}
+
+			// Parse culture names.
+			for (int32 i = MandatoryColumnNames.Num(); i < HeaderRow.Num(); ++i)
+			{
+				CultureNames.Add(HeaderRow[i]);
+			}
+
+			return true;
+		}
+
+		bool ParseWordCountRow(const FRow& WordCountRow, TMap<FString, uint32>& OutCultureWordCounts)
+		{
+			// Validate row.
+			bool IsRowValid = [&]() -> bool
+			{
+				if (WordCountRow.Num() != MandatoryColumnNames.Num() + CultureNames.Num())
+				{
+					return false;
+				}
+
+				return true;
+			}();
+
+			if (!IsRowValid)
+			{
+				return false;
+			}
+
+			// Parse time stamp.
+			FDateTime TimeStamp;
+			if (!FDateTime::Parse(WordCountRow[0], TimeStamp))
+			{
+				return false;
+			}
+
+			// Parse supported culture word counts.
+			for (int32 j = MandatoryColumnNames.Num(); j < WordCountRow.Num(); ++j)
+			{
+				const uint32 WordCount = static_cast<uint32>( FCString::Atoi(WordCountRow[j]) );
+				const int32 CultureIndex = j - MandatoryColumnNames.Num();
+				OutCultureWordCounts.FindOrAdd(CultureNames[CultureIndex]) = WordCount;
+			}
+
+			return true;
+		}
+
+	private:
+		FLocalizationTargetSettings& Target;
+		TArray<FString> MandatoryColumnNames;
+		TArray<FString> CultureNames;
+	};
+}
+
+void ULocalizationTarget::PostEditChangeProperty(FPropertyChangedEvent& PropertyChangedEvent)
+{
+	Super::PostEditChangeProperty(PropertyChangedEvent);
+
+	// Notify parent of change.
+	ULocalizationTargetSet* const LocalizationTargetSet = Cast<ULocalizationTargetSet>(GetOuter());
+	if (LocalizationTargetSet)
+	{
+		LocalizationTargetSet->PostEditChange();
+	}
+}
+
+bool ULocalizationTarget::IsMemberOfEngineTargetSet() const
+{
+	ULocalizationTargetSet* const TargetSet = CastChecked<ULocalizationTargetSet>(GetOuter());
+	return TargetSet == ULocalizationSettings::GetEngineTargetSet();
+}
+
+bool ULocalizationTarget::UpdateWordCountsFromCSV()
+{
+	const FString CSVFilePath = LocalizationConfigurationScript::GetWordCountCSVPath(this);
+	const bool Succeeded = FWordCountCSVParser::Execute(Settings, CSVFilePath);
+	if (!Succeeded)
+	{
+		for (FCultureStatistics& CultureStatistics : Settings.SupportedCulturesStatistics)
+		{
+			CultureStatistics.WordCount = 0;
+		}
+	}
+	return Succeeded;
+}
+
+void ULocalizationTarget::UpdateStatusFromConflictReport()
+{
+	const FString ConflictReportFilePath = LocalizationConfigurationScript::GetConflictReportPath(this);
+	const int64 ConflictReportFileSize = IFileManager::Get().FileSize(*ConflictReportFilePath);
+
+	switch (ConflictReportFileSize)
+	{
+	case INDEX_NONE:
+		// Unknown
+		Settings.ConflictStatus = ELocalizationTargetConflictStatus::Unknown;
+		break;
+	case 0:
+		// Good
+		Settings.ConflictStatus = ELocalizationTargetConflictStatus::Clear;
+		break;
+	default:
+		// Bad
+		Settings.ConflictStatus = ELocalizationTargetConflictStatus::ConflictsPresent;
+		break;
+	}
+}
+
+bool ULocalizationTarget::RenameTargetAndFiles(const FString& NewName)
+{
+	bool HasCompletelySucceeded = true;
+	ISourceControlModule& SourceControl = ISourceControlModule::Get();
+	ISourceControlProvider& SourceControlProvider = SourceControl.GetProvider();
+	const bool CanUseSourceControl = SourceControl.IsEnabled() && SourceControlProvider.IsEnabled() && SourceControlProvider.IsAvailable();
+
+	const auto& TryDelete = [&](const FString& Path) -> bool
+	{
+		const FSourceControlStatePtr SourceControlState = SourceControlProvider.GetState(*Path, EStateCacheUsage::Use);
+
+		// Should only use source control if the file is actually under source control.
+		if (CanUseSourceControl && SourceControlState.IsValid() && !SourceControlState->CanAdd())
+		{
+			// File is already marked for deletion.
+			if (SourceControlState->IsDeleted())
+			{
+				return true;
+			}
+			else 
+			{
+				// File is in some modified source control state, needs to be reverted.
+				if (SourceControlState->IsAdded() || SourceControlState->IsCheckedOut())
+				{
+					if (SourceControlProvider.Execute(ISourceControlOperation::Create<FRevert>(), *Path) != ECommandResult::Succeeded)
+					{
+						return false;
+					}
+				}
+
+				// File needs to be deleted.
+				if (SourceControlProvider.Execute(ISourceControlOperation::Create<FDelete>(), *Path) == ECommandResult::Succeeded)
+				{
+					return true;
+				}
+			}
+		}
+		// Attempt local deletion.
+		else if (IFileManager::Get().Delete(*Path, false, true))
+		{
+			return true;
+		}
+
+		return false;
+	};
+
+	// Delete configuration files.
+	const TArray<FString>& ScriptPaths = LocalizationConfigurationScript::GetScriptPaths(this);
+	for (const FString& ScriptPath : ScriptPaths)
+	{
+		if (!TryDelete(ScriptPath))
+		{
+			HasCompletelySucceeded = false;
+		}
+	}
+
+	const auto& GetNamedPaths = [this](TArray<FString>& NamedPaths)
+	{
+		NamedPaths.Add(LocalizationConfigurationScript::GetManifestPath(this));
+		NamedPaths.Add(LocalizationConfigurationScript::GetWordCountCSVPath(this));
+		NamedPaths.Add(LocalizationConfigurationScript::GetConflictReportPath(this));
+
+		for (const FCultureStatistics& Culture : Settings.SupportedCulturesStatistics)
+		{
+			NamedPaths.Add(LocalizationConfigurationScript::GetArchivePath(this, Culture.CultureName));
+			NamedPaths.Add(LocalizationConfigurationScript::GetDefaultPOPath(this, Culture.CultureName));
+			NamedPaths.Add(LocalizationConfigurationScript::GetLocResPath(this, Culture.CultureName));
+		}
+
+		NamedPaths.Add(LocalizationConfigurationScript::GetDataDirectory(this));
+	};
+
+	TArray<FString> OldPaths;
+	GetNamedPaths(OldPaths);
+
+	// Rename
+	Settings.Name = NewName;
+
+	TArray<FString> NewPaths;
+	GetNamedPaths(NewPaths);
+
+	// Rename data files.
+	check(OldPaths.Num() == NewPaths.Num());
+	for (int32 i = 0; i < OldPaths.Num(); ++i)
+	{
+		if (IFileManager::Get().Move(*NewPaths[i], *OldPaths[i], true, true, false, true))
+		{
+			if (CanUseSourceControl)
+			{
+				// Add new file.
+				if (SourceControlProvider.Execute(ISourceControlOperation::Create<FMarkForAdd>(), *NewPaths[i]) != ECommandResult::Succeeded)
+				{
+					HasCompletelySucceeded = false;
+				}
+			}
+
+			bool ShouldUseSourceControl = false;
+
+			if (CanUseSourceControl)
+			{
+				const FSourceControlStatePtr SourceControlState = SourceControlProvider.GetState(*OldPaths[i], EStateCacheUsage::Use);
+				if (SourceControlState->IsSourceControlled())
+				{
+					ShouldUseSourceControl = true;
+				}
+			}
+
+			// Use source control.
+			if (ShouldUseSourceControl)
+			{
+				// Delete old directory/file.
+				if (SourceControlProvider.Execute(ISourceControlOperation::Create<FDelete>(), *OldPaths[i]) != ECommandResult::Succeeded)
+				{
+					HasCompletelySucceeded = false;
+				}
+			}
+			// Operate locally.
+			else
+			{
+				// Delete old directory/file.
+				if (	(FPaths::DirectoryExists(OldPaths[i]) && !IFileManager::Get().DeleteDirectory(*OldPaths[i], false, true)) ||
+					(FPaths::FileExists(OldPaths[i]) && !IFileManager::Get().Delete(*OldPaths[i], false, true))	)
+				{
+					HasCompletelySucceeded = false;
+				}
+			}
+		}
+	}
+
+	// Generate new configuration files.
+	LocalizationConfigurationScript::GenerateGatherScript(this).Write(LocalizationConfigurationScript::GetGatherScriptPath(this));
+	LocalizationConfigurationScript::GenerateImportScript(this).Write(LocalizationConfigurationScript::GetImportScriptPath(this));
+	LocalizationConfigurationScript::GenerateExportScript(this).Write(LocalizationConfigurationScript::GetExportScriptPath(this));
+	LocalizationConfigurationScript::GenerateWordCountReportScript(this).Write(LocalizationConfigurationScript::GetWordCountReportScriptPath(this));
+
+	//if (CanUseSourceControl)
+	//{
+	//	HasCompletelySucceeded = HasCompletelySucceeded && SourceControlProvider.Execute(ISourceControlOperation::Create<FMarkForAdd>(), *LocalizationConfigurationScript::GetGatherScriptPath(this)) == ECommandResult::Succeeded;
+	//	HasCompletelySucceeded = HasCompletelySucceeded && SourceControlProvider.Execute(ISourceControlOperation::Create<FMarkForAdd>(), *LocalizationConfigurationScript::GetImportScriptPath(this)) == ECommandResult::Succeeded;
+	//	HasCompletelySucceeded = HasCompletelySucceeded && SourceControlProvider.Execute(ISourceControlOperation::Create<FMarkForAdd>(), *LocalizationConfigurationScript::GetExportScriptPath(this)) == ECommandResult::Succeeded;
+	//	HasCompletelySucceeded = HasCompletelySucceeded && SourceControlProvider.Execute(ISourceControlOperation::Create<FMarkForAdd>(), *LocalizationConfigurationScript::GetReportScriptPath(this)) == ECommandResult::Succeeded;
+	//}
+
+	return HasCompletelySucceeded;
+}
+
+bool ULocalizationTarget::DeleteFiles(const FString* const Culture) const
+{
+	bool HasCompletelySucceeded = true;
+	ISourceControlModule& SourceControl = ISourceControlModule::Get();
+	ISourceControlProvider& SourceControlProvider = SourceControl.GetProvider();
+
+	const FString DataDirectory = LocalizationConfigurationScript::GetDataDirectory(this) + (Culture ? *Culture : FString());
+
+	// Remove data files from source control.
+	{
+		TArray<FString> FilesInDataDirectory;
+		IFileManager::Get().FindFilesRecursive(FilesInDataDirectory, *DataDirectory, TEXT("*"), true, false);
+
+		TArray<FSourceControlStateRef> SourceControlStates;
+		ECommandResult::Type Result = SourceControlProvider.GetState( FilesInDataDirectory, SourceControlStates, EStateCacheUsage::ForceUpdate );
+		if (Result == ECommandResult::Succeeded)
+		{
+			TArray<FString> FilesToRevert;
+			TArray<FString> FilesToDelete;
+
+			for (const FSourceControlStateRef& SourceControlState : SourceControlStates)
+			{
+				// Added/edited files must be reverted.
+				if (SourceControlState->IsAdded() || SourceControlState->IsCheckedOut())
+				{
+					FilesToRevert.Add(SourceControlState->GetFilename());
+				}
+				// Edited files must be deleted.
+				if (SourceControlState->IsCheckedOut())
+				{
+					FilesToDelete.Add(SourceControlState->GetFilename());
+				}
+			}
+
+			const bool WasSuccessfullyReverted = (SourceControlProvider.Execute( ISourceControlOperation::Create<FRevert>(), FilesToRevert ) == ECommandResult::Succeeded);
+			const bool WasSuccessfullyDeleted = (SourceControlProvider.Execute( ISourceControlOperation::Create<FDelete>(), FilesToDelete ) == ECommandResult::Succeeded);
+			HasCompletelySucceeded = HasCompletelySucceeded && WasSuccessfullyReverted && WasSuccessfullyDeleted;
+		}
+	}
+
+	// Delete data files.
+	if (!IFileManager::Get().DeleteDirectory(*DataDirectory, false, true))
+	{
+		HasCompletelySucceeded = false;
+	}
+
+	// Delete configuration files if deleting the target.
+	if (!Culture)
+	{
+		const TArray<FString>& ScriptPaths = LocalizationConfigurationScript::GetScriptPaths(this);
+
+		// Remove script files from source control.
+		{
+			TArray<FSourceControlStateRef> SourceControlStates;
+			ECommandResult::Type Result = SourceControlProvider.GetState( ScriptPaths, SourceControlStates, EStateCacheUsage::ForceUpdate );
+			if (Result == ECommandResult::Succeeded)
+			{
+				TArray<FString> FilesToRevert;
+				TArray<FString> FilesToDelete;
+
+				for (const FSourceControlStateRef& SourceControlState : SourceControlStates)
+				{
+					// Added/edited files must be reverted.
+					if (SourceControlState->IsAdded() || SourceControlState->IsCheckedOut())
+					{
+						FilesToRevert.Add(SourceControlState->GetFilename());
+					}
+					// Edited files must be deleted.
+					if (SourceControlState->IsCheckedOut())
+					{
+						FilesToDelete.Add(SourceControlState->GetFilename());
+					}
+				}
+
+				const bool WasSuccessfullyReverted = (SourceControlProvider.Execute( ISourceControlOperation::Create<FRevert>(), FilesToRevert ) == ECommandResult::Succeeded);
+				const bool WasSuccessfullyDeleted = (SourceControlProvider.Execute( ISourceControlOperation::Create<FDelete>(), FilesToDelete ) == ECommandResult::Succeeded);
+				HasCompletelySucceeded = HasCompletelySucceeded && WasSuccessfullyReverted && WasSuccessfullyDeleted;
+			}
+		}
+
+		// Delete script files.
+		for (const FString& ScriptPath : ScriptPaths)
+		{
+			if (!IFileManager::Get().Delete(*ScriptPath, false, true))
+			{
+				HasCompletelySucceeded = false;
+			}
+		}
+	}
+
+	return HasCompletelySucceeded;
+}
+
+void ULocalizationTargetSet::PostEditChangeProperty(FPropertyChangedEvent& PropertyChangedEvent)
+{
+	Super::PostEditChangeProperty(PropertyChangedEvent);
+
+	// Notify parent of change.
+	ULocalizationSettings* const LocalizationSettings = Cast<ULocalizationSettings>(GetOuter());
+	if (LocalizationSettings)
+	{
+		LocalizationSettings->PostEditChange();
+	}
+}
+
+#undef LOCTEXT_NAMESPACE
\ No newline at end of file
diff --git a/Engine/Source/Editor/Localization/Public/LocalizationConfigurationScript.h b/Engine/Source/Editor/Localization/Public/LocalizationConfigurationScript.h
new file mode 100644
--- /dev/null
+++ b/Engine/Source/Editor/Localization/Public/LocalizationConfigurationScript.h
@@ -0,0 +1,54 @@
+// Copyright 1998-2015 Epic Games, Inc. All Rights Reserved.
+
+#pragma once
+
+#include "ConfigCacheIni.h"
+#include "LocalizationTargetTypes.h"
+
+struct LOCALIZATION_API FLocalizationConfigurationScript : public FConfigFile
+{
+	FConfigSection& CommonSettings()
+	{
+		return FindOrAdd(TEXT("CommonSettings"));
+	}
+
+	FConfigSection& GatherTextStep(const uint32 Index)
+	{
+		return FindOrAdd( FString::Printf( TEXT("GatherTextStep%u"), Index) );
+	}
+};
+
+namespace LocalizationConfigurationScript
+{
+	LOCALIZATION_API FString MakePathRelativeForCommandletProcess(const FString& Path, const bool IsUsingProjectFile);
+
+	LOCALIZATION_API FString GetScriptDirectory(const ULocalizationTarget* const Target);
+	LOCALIZATION_API FString GetDataDirectory(const ULocalizationTarget* const Target);
+	LOCALIZATION_API TArray<FString> GetScriptPaths(const ULocalizationTarget* const Target);
+
+	LOCALIZATION_API FString GetManifestPath(const ULocalizationTarget* const Target);
+	LOCALIZATION_API FString GetArchivePath(const ULocalizationTarget* const Target, const FString& CultureName);
+	LOCALIZATION_API FString GetDefaultPOFileName(const ULocalizationTarget* const Target);
+	LOCALIZATION_API FString GetDefaultPOPath(const ULocalizationTarget* const Target, const FString& CultureName);
+	LOCALIZATION_API FString GetLocResPath(const ULocalizationTarget* const Target, const FString& CultureName);
+	LOCALIZATION_API FString GetWordCountCSVPath(const ULocalizationTarget* const Target);
+	LOCALIZATION_API FString GetConflictReportPath(const ULocalizationTarget* const Target);
+
+	LOCALIZATION_API FLocalizationConfigurationScript GenerateGatherScript(const ULocalizationTarget* const Target);
+	LOCALIZATION_API FString GetGatherScriptPath(const ULocalizationTarget* const Target);
+
+	LOCALIZATION_API FLocalizationConfigurationScript GenerateImportScript(const ULocalizationTarget* const Target, const TOptional<FString> CultureName = TOptional<FString>(), const TOptional<FString> OutputPathOverride = TOptional<FString>());
+	LOCALIZATION_API FString GetImportScriptPath(const ULocalizationTarget* const Target, const TOptional<FString> CultureName = TOptional<FString>());
+
+	LOCALIZATION_API FLocalizationConfigurationScript GenerateExportScript(const ULocalizationTarget* const Target, const TOptional<FString> CultureName = TOptional<FString>(), const TOptional<FString> OutputPathOverride = TOptional<FString>());
+	LOCALIZATION_API FString GetExportScriptPath(const ULocalizationTarget* const Target, const TOptional<FString> CultureName = TOptional<FString>());
+
+	LOCALIZATION_API FLocalizationConfigurationScript GenerateWordCountReportScript(const ULocalizationTarget* const Target);
+	LOCALIZATION_API FString GetWordCountReportScriptPath(const ULocalizationTarget* const Target);
+
+	LOCALIZATION_API FLocalizationConfigurationScript GenerateCompileScript(const ULocalizationTarget* const Target, const TOptional<FString> CultureName = TOptional<FString>());
+	LOCALIZATION_API FString GetCompileScriptPath(const ULocalizationTarget* const Target, const TOptional<FString> CultureName = TOptional<FString>());
+
+	LOCALIZATION_API FLocalizationConfigurationScript GenerateRegenerateResourcesScript(const ULocalizationTarget* const Target);
+	LOCALIZATION_API FString GetRegenerateResourcesScriptPath(const ULocalizationTarget* const Target);
+}
\ No newline at end of file
diff --git a/Engine/Source/Editor/Localization/Public/LocalizationModule.h b/Engine/Source/Editor/Localization/Public/LocalizationModule.h
new file mode 100644
--- /dev/null
+++ b/Engine/Source/Editor/Localization/Public/LocalizationModule.h
@@ -0,0 +1,11 @@
+// Copyright 1998-2015 Epic Games, Inc. All Rights Reserved.
+
+#include "ModuleInterface.h"
+
+class ULocalizationTarget;
+
+class FLocalizationModule : public IModuleInterface
+{
+public:
+	ULocalizationTarget* GetLocalizationTargetByName(FString TargetName, bool bIsEngineTarget);
+};
\ No newline at end of file
diff --git a/Engine/Source/Editor/Localization/Public/LocalizationSettings.h b/Engine/Source/Editor/Localization/Public/LocalizationSettings.h
new file mode 100644
--- /dev/null
+++ b/Engine/Source/Editor/Localization/Public/LocalizationSettings.h
@@ -0,0 +1,39 @@
+// Copyright 1998-2015 Epic Games, Inc. All Rights Reserved.
+
+#pragma once
+
+#include "Object.h"
+#include "LocalizationTargetTypes.h"
+#include "LocalizationSettings.generated.h"
+
+// Class for loading/saving configuration settings and the details view objects needed for localization dashboard functionality.
+UCLASS(Config=Editor, defaultconfig)
+class LOCALIZATION_API ULocalizationSettings : public UObject
+{
+	GENERATED_BODY()
+
+public:
+	ULocalizationSettings(const FObjectInitializer& ObjectInitializer);
+
+private:
+	UPROPERTY()
+	ULocalizationTargetSet* EngineTargetSet;
+
+	UPROPERTY(config)
+	TArray<FLocalizationTargetSettings> EngineTargetsSettings;
+
+	UPROPERTY()
+	ULocalizationTargetSet* GameTargetSet;
+
+	UPROPERTY(config)
+	TArray<FLocalizationTargetSettings> GameTargetsSettings;
+
+public:
+#if WITH_EDITOR
+	virtual void PostInitProperties() override;
+	virtual void PostEditChangeProperty(FPropertyChangedEvent& PropertyChangedEvent) override;
+#endif
+
+	static ULocalizationTargetSet* GetEngineTargetSet();
+	static ULocalizationTargetSet* GetGameTargetSet();
+};
\ No newline at end of file
diff --git a/Engine/Source/Editor/Localization/Public/LocalizationTargetTypes.h b/Engine/Source/Editor/Localization/Public/LocalizationTargetTypes.h
new file mode 100644
--- /dev/null
+++ b/Engine/Source/Editor/Localization/Public/LocalizationTargetTypes.h
@@ -0,0 +1,337 @@
+﻿// Copyright 1998-2015 Epic Games, Inc. All Rights Reserved.
+
+#pragma once
+
+#include "Engine/EngineTypes.h"
+#include "LocalizationTargetTypes.generated.h"
+
+UENUM()
+enum class ELocalizationTargetConflictStatus : uint8
+{
+	/** The status of conflicts in this localization target could not be determined. */
+	Unknown,
+	/** The are outstanding conflicts present in this localization target. */
+	ConflictsPresent,
+	/** The localization target is clear of conflicts. */
+	Clear,
+};
+
+USTRUCT()
+struct FGatherTextSearchDirectory
+{
+	GENERATED_USTRUCT_BODY()
+
+	UPROPERTY(config, EditAnywhere, Category="Path")
+	FString Path;
+
+	LOCALIZATION_API bool Validate(const FString& RootDirectory, FText& OutError) const;
+};
+
+USTRUCT()
+struct FGatherTextIncludePath
+{
+	GENERATED_USTRUCT_BODY()
+
+	UPROPERTY(config, EditAnywhere, Category="Pattern")
+	FString Pattern;
+
+	LOCALIZATION_API bool Validate(const FString& RootDirectory, FText& OutError) const;
+};
+
+USTRUCT()
+struct FGatherTextExcludePath
+{
+	GENERATED_USTRUCT_BODY()
+
+	UPROPERTY(config, EditAnywhere, Category="Pattern")
+	FString Pattern;
+
+	LOCALIZATION_API bool Validate(FText& OutError) const;
+};
+
+USTRUCT()
+struct FGatherTextFileExtension
+{
+	GENERATED_USTRUCT_BODY()
+
+	UPROPERTY(config, EditAnywhere, Category="Pattern")
+	FString Pattern;
+
+	LOCALIZATION_API bool Validate(FText& OutError) const;
+};
+
+USTRUCT()
+struct FGatherTextFromTextFilesConfiguration
+{
+	GENERATED_USTRUCT_BODY()
+
+	static const TArray<FGatherTextFileExtension> DefaultTextFileExtensions;
+
+	FGatherTextFromTextFilesConfiguration()
+		: IsEnabled(true)
+		, FileExtensions(DefaultTextFileExtensions)
+	{
+	}
+
+	/* If enabled, text from text files will be gathered according to this configuration. */
+	UPROPERTY(config, EditAnywhere, Category = "Execution")
+	bool IsEnabled;
+
+	/* The paths of directories to be searched recursively for text files, specified relative to the project's root, which may be parsed for text to gather. */
+	UPROPERTY(config, EditAnywhere, Category = "Filter")
+	TArray<FGatherTextSearchDirectory> SearchDirectories;
+
+	/* Text files whose paths match these wildcard patterns will be excluded from gathering. */
+	UPROPERTY(config, EditAnywhere, Category = "Filter")
+	TArray<FGatherTextExcludePath> ExcludePathWildcards;
+
+	/* Text files whose names match these wildcard patterns may be parsed for text to gather. */
+	UPROPERTY(config, EditAnywhere, Category = "Filter")
+	TArray<FGatherTextFileExtension> FileExtensions;
+
+	LOCALIZATION_API bool Validate(const FString& RootDirectory, FText& OutError) const;
+};
+
+
+USTRUCT()
+struct FGatherTextFromPackagesConfiguration
+{
+	GENERATED_USTRUCT_BODY()
+
+	static const TArray<FGatherTextFileExtension> DefaultPackageFileExtensions;
+
+	FGatherTextFromPackagesConfiguration()
+		: IsEnabled(true)
+		, FileExtensions(DefaultPackageFileExtensions)
+	{
+	}
+
+	/* If enabled, text from packages will be gathered according to this configuration. */
+	UPROPERTY(config, EditAnywhere, Category = "Execution")
+	bool IsEnabled;
+
+	/* Packages whose paths match these wildcard patterns, specified relative to the project's root, may be processed for gathering. */
+	UPROPERTY(config, EditAnywhere, Category = "Filter")
+	TArray<FGatherTextIncludePath> IncludePathWildcards;
+
+	/* Packages whose paths match these wildcard patterns will be excluded from gathering. */
+	UPROPERTY(config, EditAnywhere, Category = "Filter")
+	TArray<FGatherTextExcludePath> ExcludePathWildcards;
+
+	/* Packages whose names match these wildcard patterns may be processed for text to gather. */
+	UPROPERTY(config, EditAnywhere, Category = "Filter")
+	TArray<FGatherTextFileExtension> FileExtensions;
+
+	LOCALIZATION_API bool Validate(const FString& RootDirectory, FText& OutError) const;
+};
+
+USTRUCT()
+struct FMetaDataTextKeyPattern
+{
+	GENERATED_USTRUCT_BODY()
+
+	UPROPERTY(config, EditAnywhere, Category="Pattern")
+	FString Pattern;
+
+	LOCALIZATION_API bool Validate(FText& OutError) const;
+
+	LOCALIZATION_API static const TArray<FString>& GetPossiblePlaceHolders();
+};
+
+USTRUCT()
+struct FMetaDataKeyName
+{
+	GENERATED_USTRUCT_BODY()
+
+	UPROPERTY(config, EditAnywhere, Category="Name")
+	FString Name;
+
+	LOCALIZATION_API bool Validate(FText& OutError) const;
+};
+
+USTRUCT()
+struct FMetaDataKeyGatherSpecification
+{
+	GENERATED_USTRUCT_BODY()
+
+	/*  The metadata key for which values will be gathered as text. */
+	UPROPERTY(config, EditAnywhere, Category = "Input")
+	FMetaDataKeyName MetaDataKey;
+
+	/* The localization namespace in which the gathered text will be output. */
+	UPROPERTY(config, EditAnywhere, Category = "Output")
+	FString TextNamespace;
+
+	/* The pattern which will be formatted to form the localization key for the metadata value gathered as text.
+	Placeholder - Description
+	{FieldPath} - The fully qualified name of the object upon which the metadata resides.
+	{MetaDataValue} - The value associated with the metadata key. */
+	UPROPERTY(config, EditAnywhere, Category = "Output")
+	FMetaDataTextKeyPattern TextKeyPattern;
+
+	LOCALIZATION_API bool Validate(FText& OutError) const;
+};
+
+USTRUCT()
+struct FGatherTextFromMetaDataConfiguration
+{
+	GENERATED_USTRUCT_BODY()
+
+	FGatherTextFromMetaDataConfiguration()
+		: IsEnabled(false)
+	{
+	}
+
+	/* If enabled, metadata will be gathered according to this configuration. */
+	UPROPERTY(config, EditAnywhere, Category = "Execution")
+	bool IsEnabled;
+
+	/* Metadata from source files whose paths match these wildcard patterns, specified relative to the project's root, may be processed for gathering. */
+	UPROPERTY(config, EditAnywhere, Category = "Filter")
+	TArray<FGatherTextIncludePath> IncludePathWildcards;
+
+	/* Metadata from source files whose paths match these wildcard patterns will be excluded from gathering. */
+	UPROPERTY(config, EditAnywhere, Category = "Filter")
+	TArray<FGatherTextExcludePath> ExcludePathWildcards;
+
+	/* Specifications for how to gather text from specific metadata keys. */
+	UPROPERTY(config, EditAnywhere, Category = "MetaData")
+	TArray<FMetaDataKeyGatherSpecification> KeySpecifications;
+
+	LOCALIZATION_API bool Validate(const FString& RootDirectory, FText& OutError) const;
+};
+
+USTRUCT()
+struct FCultureStatistics
+{
+	GENERATED_USTRUCT_BODY()
+
+	FCultureStatistics()
+		: WordCount(0)
+	{
+	}
+
+	FCultureStatistics(const FString& InCultureName)
+		: CultureName(InCultureName)
+		, WordCount(0)
+	{
+	}
+
+
+	/* The ISO name for this culture. */
+	UPROPERTY(config, EditAnywhere, Category = "Culture")
+	FString CultureName;
+
+	/* The estimated number of words that have been localized for this culture. */
+	UPROPERTY(Transient, EditAnywhere, Category = "Statistics")
+	uint32 WordCount;
+};
+
+UENUM()
+enum class ELocalizationTargetLoadingPolicy : uint8
+{
+	/** This target's localization data will never be loaded automatically. */
+	Never,
+	/** This target's localization data will always be loaded automatically. */
+	Always,
+	/** This target's localization data will only be loaded when running the editor. Use if this target localizes the editor. */
+	Editor,
+	/** This target's localization data will only be loaded when running the game. Use if this target localizes your game. */
+	Game,
+	/** This target's localization data will only be loaded if the editor is displaying localized property names. */
+	PropertyNames,
+	/** This target's localization data will only be loaded if the editor is displaying localized tool tips. */
+	ToolTips,
+};
+
+USTRUCT()
+struct FLocalizationTargetSettings
+{
+	GENERATED_USTRUCT_BODY()
+
+	FLocalizationTargetSettings()
+		: Guid(FGuid::NewGuid())
+		, ConflictStatus(ELocalizationTargetConflictStatus::Unknown)
+		, NativeCultureIndex(INDEX_NONE)
+	{
+	}
+
+	/* Unique name for the target. */
+	UPROPERTY(config, EditAnywhere, Category = "Target")
+	FString Name;
+
+	UPROPERTY(config)
+	FGuid Guid;
+
+	/* Whether the target has outstanding conflicts that require resolution. */
+	UPROPERTY(Transient, EditAnywhere, Category = "Target")
+	ELocalizationTargetConflictStatus ConflictStatus;
+
+	/* Text present in these targets will not be duplicated in this target. */
+	UPROPERTY(config, EditAnywhere, Category = "Gather")
+	TArray<FGuid> TargetDependencies;
+
+	/* Text present in these manifests will not be duplicated in this target. */
+	UPROPERTY(config, EditAnywhere, Category = "Gather", AdvancedDisplay, meta=(FilePathFilter="manifest"))
+	TArray<FFilePath> AdditionalManifestDependencies;
+
+	/* The names of modules which must be loaded when gathering text. Use to gather from packages or metadata sourced from a module that isn't the primary game module. */
+	UPROPERTY(config, EditAnywhere, Category = "Gather", AdvancedDisplay)
+	TArray<FString> RequiredModuleNames;
+
+	/* Parameters for defining what text is gathered from text files and how. */
+	UPROPERTY(config, EditAnywhere, Category = "Gather")
+	FGatherTextFromTextFilesConfiguration GatherFromTextFiles;
+
+	/* Parameters for defining what text is gathered from packages and how. */
+	UPROPERTY(config, EditAnywhere, Category = "Gather")
+	FGatherTextFromPackagesConfiguration GatherFromPackages;
+
+	/* Parameters for defining what text is gathered from metadata and how. */
+	UPROPERTY(config, EditAnywhere, Category = "Gather")
+	FGatherTextFromMetaDataConfiguration GatherFromMetaData;
+
+	/* The index of the native culture among the supported cultures. */
+	UPROPERTY(config, EditAnywhere, Category = "Cultures")
+	int32 NativeCultureIndex;
+
+	/* Cultures for which the source text is being localized for.*/
+	UPROPERTY(config, EditAnywhere, Category = "Cultures")
+	TArray<FCultureStatistics> SupportedCulturesStatistics;
+};
+
+UCLASS(Within=LocalizationTargetSet)
+class LOCALIZATION_API ULocalizationTarget : public UObject
+{
+	GENERATED_BODY()
+
+public:
+	UPROPERTY(EditAnywhere, Category = "Target")
+	FLocalizationTargetSettings Settings;
+
+public:
+#if WITH_EDITOR
+	virtual void PostEditChangeProperty(FPropertyChangedEvent& PropertyChangedEvent) override;
+#endif
+
+	bool IsMemberOfEngineTargetSet() const;
+	bool UpdateWordCountsFromCSV();
+	void UpdateStatusFromConflictReport();
+	bool RenameTargetAndFiles(const FString& NewName);
+	bool DeleteFiles(const FString* const Culture = nullptr) const;
+};
+
+UCLASS(Within=LocalizationSettings)
+class LOCALIZATION_API ULocalizationTargetSet : public UObject
+{
+	GENERATED_BODY()
+
+public:
+	UPROPERTY(EditAnywhere, Category = "Targets")
+	TArray<ULocalizationTarget*> TargetObjects;
+
+public:
+#if WITH_EDITOR
+	virtual void PostEditChangeProperty(FPropertyChangedEvent& PropertyChangedEvent) override;
+#endif
+};
\ No newline at end of file
diff --git a/Engine/Source/Editor/LocalizationDashboard/LocalizationDashboard.Build.cs b/Engine/Source/Editor/LocalizationDashboard/LocalizationDashboard.Build.cs
--- a/Engine/Source/Editor/LocalizationDashboard/LocalizationDashboard.Build.cs
+++ b/Engine/Source/Editor/LocalizationDashboard/LocalizationDashboard.Build.cs
@@ -12,6 +12,7 @@ public class LocalizationDashboard : ModuleRules
 				"Core",
                 "CoreUObject",
                 "PropertyEditor",
+                "Localization"
 			}
 		);
 
@@ -29,6 +30,7 @@ public class LocalizationDashboard : ModuleRules
                 "MainFrame",
                 "SourceControl",
                 "SharedSettingsWidgets",
+                "Localization",
                 "LocalizationService"
 			}
 		);
diff --git a/Engine/Source/Editor/LocalizationDashboard/Private/GatherTextDetailCustomizations.cpp b/Engine/Source/Editor/LocalizationDashboard/Private/GatherTextDetailCustomizations.cpp
--- a/Engine/Source/Editor/LocalizationDashboard/Private/GatherTextDetailCustomizations.cpp
+++ b/Engine/Source/Editor/LocalizationDashboard/Private/GatherTextDetailCustomizations.cpp
@@ -641,7 +641,7 @@ namespace
 	{
 		PatternPropertyHandle = InPatternPropertyHandle;
 
-		for(const FString& PossiblePlaceHolder : FMetaDataTextKeyPattern::PossiblePlaceHolders)
+		for(const FString& PossiblePlaceHolder : FMetaDataTextKeyPattern::GetPossiblePlaceHolders())
 		{
 			PossiblePlaceHoldersListItemsSource.Add( MakeShareable(new FString(PossiblePlaceHolder)) );
 		}
diff --git a/Engine/Source/Editor/LocalizationDashboard/Private/LocalizationDashboard.cpp b/Engine/Source/Editor/LocalizationDashboard/Private/LocalizationDashboard.cpp
--- a/Engine/Source/Editor/LocalizationDashboard/Private/LocalizationDashboard.cpp
+++ b/Engine/Source/Editor/LocalizationDashboard/Private/LocalizationDashboard.cpp
@@ -8,7 +8,7 @@
 #include "IMainFrameModule.h"
 #include "ModuleManager.h"
 #include "SLocalizationTargetEditor.h"
-#include "LocalizationDashboardSettings.h"
+#include "LocalizationSettings.h"
 #include "LocalizationTargetTypes.h"
 #include "SSettingsEditorCheckoutNotice.h"
 
@@ -60,12 +60,12 @@ void SLocalizationDashboard::Construct(const FArguments& InArgs, const TSharedPt
 		if (SpawnTabArgs.GetTabId() == EngineTargetSetDetailsTabName)
 		{
 			DockTabLabel = LOCTEXT("EngineTargetsTabLabel", "Engine Targets");
-			TargetSet = ULocalizationDashboardSettings::GetEngineTargetSet();
+			TargetSet = ULocalizationSettings::GetEngineTargetSet();
 		}
 		else if (SpawnTabArgs.GetTabId() == ProjectTargetSetDetailsTabName)
 		{
 			DockTabLabel = LOCTEXT("ProjectTargetsTabLabel", "Project Targets");
-			TargetSet = ULocalizationDashboardSettings::GetGameTargetSet();
+			TargetSet = ULocalizationSettings::GetGameTargetSet();
 		}
 
 		const TSharedRef<SDockTab> DockTab = SNew(SDockTab)
@@ -117,7 +117,7 @@ void SLocalizationDashboard::Construct(const FArguments& InArgs, const TSharedPt
 	const TSharedRef<SWidget> MenuWidget = MainFrameModule.MakeMainMenu( TabManager, MenuExtender );
 
 	FString ConfigFilePath;
-	ConfigFilePath = GetDefault<ULocalizationDashboardSettings>()->GetDefaultConfigFilename();
+	ConfigFilePath = GetDefault<ULocalizationSettings>()->GetDefaultConfigFilename();
 	ConfigFilePath = FPaths::ConvertRelativePathToFull(ConfigFilePath);
 
 	ChildSlot
diff --git a/Engine/Source/Editor/LocalizationDashboard/Private/LocalizationDashboardModule.cpp b/Engine/Source/Editor/LocalizationDashboard/Private/LocalizationDashboardModule.cpp
--- a/Engine/Source/Editor/LocalizationDashboard/Private/LocalizationDashboardModule.cpp
+++ b/Engine/Source/Editor/LocalizationDashboard/Private/LocalizationDashboardModule.cpp
@@ -12,7 +12,7 @@
 #include "GatherTextDetailCustomizations.h"
 #include "PropertyEditorModule.h"
 #include "ILocalizationServiceModule.h"
-#include "LocalizationDashboardSettings.h"
+#include "LocalizationSettings.h"
 
 #define LOCTEXT_NAMESPACE "LocalizationDashboard"
 
@@ -157,7 +157,7 @@ public:
 	{
 		if (bIsEngineTarget)
 		{
-			ULocalizationTargetSet* EngineTargetSet = ULocalizationDashboardSettings::GetEngineTargetSet();
+			ULocalizationTargetSet* EngineTargetSet = ULocalizationSettings::GetEngineTargetSet();
 			for (ULocalizationTarget* Target : EngineTargetSet->TargetObjects)
 			{
 				if (Target->Settings.Name == TargetName)
@@ -168,7 +168,7 @@ public:
 		}
 		else
 		{
-			ULocalizationTargetSet* EngineTargetSet = ULocalizationDashboardSettings::GetEngineTargetSet();
+			ULocalizationTargetSet* EngineTargetSet = ULocalizationSettings::GetEngineTargetSet();
 			for (ULocalizationTarget* Target : EngineTargetSet->TargetObjects)
 			{
 				if (Target->Settings.Name == TargetName)
diff --git a/Engine/Source/Editor/LocalizationDashboard/Private/LocalizationTargetDetailCustomization.cpp b/Engine/Source/Editor/LocalizationDashboard/Private/LocalizationTargetDetailCustomization.cpp
--- a/Engine/Source/Editor/LocalizationDashboard/Private/LocalizationTargetDetailCustomization.cpp
+++ b/Engine/Source/Editor/LocalizationDashboard/Private/LocalizationTargetDetailCustomization.cpp
@@ -13,7 +13,7 @@
 #include "IDetailsView.h"
 #include "LocalizationCommandletTasks.h"
 #include "ObjectEditorUtils.h"
-#include "LocalizationDashboardSettings.h"
+#include "LocalizationSettings.h"
 #include "SErrorText.h"
 #include "ILocalizationServiceModule.h"
 #include "ILocalizationServiceProvider.h"
@@ -393,7 +393,7 @@ FText FLocalizationTargetDetailCustomization::GetTargetName() const
 bool FLocalizationTargetDetailCustomization::IsTargetNameUnique(const FString& Name) const
 {
 	TArray<ULocalizationTarget*> AllLocalizationTargets;
-	ULocalizationTargetSet* EngineTargetSet = ULocalizationDashboardSettings::GetEngineTargetSet();
+	ULocalizationTargetSet* EngineTargetSet = ULocalizationSettings::GetEngineTargetSet();
 	if (EngineTargetSet != TargetSet)
 	{
 		AllLocalizationTargets.Append(EngineTargetSet->TargetObjects);
@@ -537,7 +537,7 @@ void FLocalizationTargetDetailCustomization::RebuildTargetDependenciesBox()
 		TargetDependenciesWidgets.Empty();
 
 		TArray<ULocalizationTarget*> AllLocalizationTargets;
-		ULocalizationTargetSet* EngineTargetSet = ULocalizationDashboardSettings::GetEngineTargetSet();
+		ULocalizationTargetSet* EngineTargetSet = ULocalizationSettings::GetEngineTargetSet();
 		if (EngineTargetSet != TargetSet)
 		{
 			AllLocalizationTargets.Append(EngineTargetSet->TargetObjects);
@@ -591,7 +591,7 @@ void FLocalizationTargetDetailCustomization::RebuildTargetsList()
 	};
 
 	TArray<ULocalizationTarget*> AllLocalizationTargets;
-	ULocalizationTargetSet* EngineTargetSet = ULocalizationDashboardSettings::GetEngineTargetSet();
+	ULocalizationTargetSet* EngineTargetSet = ULocalizationSettings::GetEngineTargetSet();
 	if (EngineTargetSet != TargetSet)
 	{
 		AllLocalizationTargets.Append(EngineTargetSet->TargetObjects);
diff --git a/Engine/Source/Editor/LocalizationDashboard/Public/ILocalizationDashboardModule.h b/Engine/Source/Editor/LocalizationDashboard/Public/ILocalizationDashboardModule.h
--- a/Engine/Source/Editor/LocalizationDashboard/Public/ILocalizationDashboardModule.h
+++ b/Engine/Source/Editor/LocalizationDashboard/Public/ILocalizationDashboardModule.h
@@ -3,9 +3,9 @@
 #pragma once
 
 #include "ModuleInterface.h"
-#include "LocalizationTargetTypes.h"
 
 class ILocalizationServiceProvider;
+class ULocalizationTarget;
 
 /**
  * Interface for localization dashboard module.
diff --git a/Engine/Source/Editor/LocalizationDashboard/Public/LocalizationConfigurationScript.h b/Engine/Source/Editor/LocalizationDashboard/Public/LocalizationConfigurationScript.h
deleted file mode 100644
diff --git a/Engine/Source/Editor/LocalizationDashboard/Public/LocalizationTargetTypes.h b/Engine/Source/Editor/LocalizationDashboard/Public/LocalizationTargetTypes.h
deleted file mode 100644
diff --git a/Engine/Source/Editor/TranslationEditor/Private/TranslationDataManager.cpp b/Engine/Source/Editor/TranslationEditor/Private/TranslationDataManager.cpp
--- a/Engine/Source/Editor/TranslationEditor/Private/TranslationDataManager.cpp
+++ b/Engine/Source/Editor/TranslationEditor/Private/TranslationDataManager.cpp
@@ -13,8 +13,8 @@
 #include "NotificationManager.h"
 #include "PortableObjectFormatDOM.h"
 #include "ILocalizationServiceModule.h"
-#include "Editor/LocalizationDashboard/Public/ILocalizationDashboardModule.h"
-#include "Editor/LocalizationDashboard/Public/LocalizationTargetTypes.h"
+#include "ILocalizationDashboardModule.h"
+#include "LocalizationTargetTypes.h"
 #include "LocalizationConfigurationScript.h"
 
 
diff --git a/Engine/Source/Editor/TranslationEditor/TranslationEditor.Build.cs b/Engine/Source/Editor/TranslationEditor/TranslationEditor.Build.cs
--- a/Engine/Source/Editor/TranslationEditor/TranslationEditor.Build.cs
+++ b/Engine/Source/Editor/TranslationEditor/TranslationEditor.Build.cs
@@ -35,6 +35,7 @@ public class TranslationEditor : ModuleRules
 				"SourceControl",
                 "MessageLog",
                 "Documentation",
+                "Localization",
                 "LocalizationService",
                 "LocalizationDashboard",
 			}
