Fixed build error due to a static member variable in some exposed structures from the new editor Localization module.

diff --git a/Engine/Source/Editor/Localization/Private/LocalizationTargetTypes.cpp b/Engine/Source/Editor/Localization/Private/LocalizationTargetTypes.cpp
--- a/Engine/Source/Editor/Localization/Private/LocalizationTargetTypes.cpp
+++ b/Engine/Source/Editor/Localization/Private/LocalizationTargetTypes.cpp
@@ -77,22 +77,22 @@ bool FGatherTextFileExtension::Validate(FText& OutError) const
 	return true;
 }
 
-const TArray<FGatherTextFileExtension> FGatherTextFromTextFilesConfiguration::DefaultTextFileExtensions = []()
+const TArray<FGatherTextFileExtension>& FGatherTextFromTextFilesConfiguration::GetDefaultTextFileExtensions()
 {
 	TArray<FGatherTextFileExtension> Result;
 	Result.Add(FGatherTextFileExtension{TEXT("h")});
 	Result.Add(FGatherTextFileExtension{TEXT("cpp")});
 	Result.Add(FGatherTextFileExtension{TEXT("ini")});
 	return Result;
-}();
+}
 
-const TArray<FGatherTextFileExtension> FGatherTextFromPackagesConfiguration::DefaultPackageFileExtensions = []()
+const TArray<FGatherTextFileExtension>& FGatherTextFromPackagesConfiguration::GetDefaultPackageFileExtensions()
 {
 	TArray<FGatherTextFileExtension> Result;
 	Result.Add(FGatherTextFileExtension{TEXT("umap")});
 	Result.Add(FGatherTextFileExtension{TEXT("uasset")});
 	return Result;
-}();
+}
 
 bool FGatherTextFromTextFilesConfiguration::Validate(const FString& RootDirectory, FText& OutError) const
 {
diff --git a/Engine/Source/Editor/Localization/Public/LocalizationTargetTypes.h b/Engine/Source/Editor/Localization/Public/LocalizationTargetTypes.h
--- a/Engine/Source/Editor/Localization/Public/LocalizationTargetTypes.h
+++ b/Engine/Source/Editor/Localization/Public/LocalizationTargetTypes.h
@@ -65,11 +65,11 @@ struct FGatherTextFromTextFilesConfiguration
 {
 	GENERATED_USTRUCT_BODY()
 
-	static const TArray<FGatherTextFileExtension> DefaultTextFileExtensions;
+	static const TArray<FGatherTextFileExtension>& GetDefaultTextFileExtensions();
 
 	FGatherTextFromTextFilesConfiguration()
 		: IsEnabled(true)
-		, FileExtensions(DefaultTextFileExtensions)
+		, FileExtensions(GetDefaultTextFileExtensions())
 	{
 	}
 
@@ -98,11 +98,11 @@ struct FGatherTextFromPackagesConfiguration
 {
 	GENERATED_USTRUCT_BODY()
 
-	static const TArray<FGatherTextFileExtension> DefaultPackageFileExtensions;
+	static const TArray<FGatherTextFileExtension>& GetDefaultPackageFileExtensions();
 
 	FGatherTextFromPackagesConfiguration()
 		: IsEnabled(true)
-		, FileExtensions(DefaultPackageFileExtensions)
+		, FileExtensions(GetDefaultPackageFileExtensions())
 	{
 	}
 
