Fixed return of local variable rather than static const variable in LocalizationTargetTypes.cpp.

diff --git a/Engine/Source/Editor/Localization/Private/LocalizationTargetTypes.cpp b/Engine/Source/Editor/Localization/Private/LocalizationTargetTypes.cpp
--- a/Engine/Source/Editor/Localization/Private/LocalizationTargetTypes.cpp
+++ b/Engine/Source/Editor/Localization/Private/LocalizationTargetTypes.cpp
@@ -79,18 +79,26 @@ bool FGatherTextFileExtension::Validate(FText& OutError) const
 
 const TArray<FGatherTextFileExtension>& FGatherTextFromTextFilesConfiguration::GetDefaultTextFileExtensions()
 {
-	TArray<FGatherTextFileExtension> Result;
-	Result.Add(FGatherTextFileExtension{TEXT("h")});
-	Result.Add(FGatherTextFileExtension{TEXT("cpp")});
-	Result.Add(FGatherTextFileExtension{TEXT("ini")});
+	static const TArray<FGatherTextFileExtension> Result = []()
+		{
+			TArray<FGatherTextFileExtension> Result;
+			Result.Add(FGatherTextFileExtension{TEXT("h")});
+			Result.Add(FGatherTextFileExtension{TEXT("cpp")});
+			Result.Add(FGatherTextFileExtension{TEXT("ini")});
+			return Result;
+		}();
 	return Result;
 }
 
 const TArray<FGatherTextFileExtension>& FGatherTextFromPackagesConfiguration::GetDefaultPackageFileExtensions()
 {
-	TArray<FGatherTextFileExtension> Result;
-	Result.Add(FGatherTextFileExtension{TEXT("umap")});
-	Result.Add(FGatherTextFileExtension{TEXT("uasset")});
+	static const TArray<FGatherTextFileExtension> Result = []()
+		{
+			TArray<FGatherTextFileExtension> Result;
+			Result.Add(FGatherTextFileExtension{TEXT("umap")});
+			Result.Add(FGatherTextFileExtension{TEXT("uasset")});
+			return Result;
+		}();
 	return Result;
 }
 
