DLL export of methods blocking compilation of Localization Dashboard module.

diff --git a/Engine/Source/Editor/Localization/Public/LocalizationTargetTypes.h b/Engine/Source/Editor/Localization/Public/LocalizationTargetTypes.h
--- a/Engine/Source/Editor/Localization/Public/LocalizationTargetTypes.h
+++ b/Engine/Source/Editor/Localization/Public/LocalizationTargetTypes.h
@@ -65,7 +65,7 @@ struct FGatherTextFromTextFilesConfiguration
 {
 	GENERATED_USTRUCT_BODY()
 
-	static const TArray<FGatherTextFileExtension>& GetDefaultTextFileExtensions();
+	LOCALIZATION_API static const TArray<FGatherTextFileExtension>& GetDefaultTextFileExtensions();
 
 	FGatherTextFromTextFilesConfiguration()
 		: IsEnabled(true)
@@ -98,7 +98,7 @@ struct FGatherTextFromPackagesConfiguration
 {
 	GENERATED_USTRUCT_BODY()
 
-	static const TArray<FGatherTextFileExtension>& GetDefaultPackageFileExtensions();
+	LOCALIZATION_API static const TArray<FGatherTextFileExtension>& GetDefaultPackageFileExtensions();
 
 	FGatherTextFromPackagesConfiguration()
 		: IsEnabled(true)
