Fixed shadowed variable warning in LocalizationTargetTypes.cpp.

diff --git a/Engine/Source/Editor/Localization/Private/LocalizationTargetTypes.cpp b/Engine/Source/Editor/Localization/Private/LocalizationTargetTypes.cpp
--- a/Engine/Source/Editor/Localization/Private/LocalizationTargetTypes.cpp
+++ b/Engine/Source/Editor/Localization/Private/LocalizationTargetTypes.cpp
@@ -79,7 +79,7 @@ bool FGatherTextFileExtension::Validate(FText& OutError) const
 
 const TArray<FGatherTextFileExtension>& FGatherTextFromTextFilesConfiguration::GetDefaultTextFileExtensions()
 {
-	static const TArray<FGatherTextFileExtension> Result = []()
+	static const TArray<FGatherTextFileExtension> DefaultTextFileExtensions = []()
 		{
 			TArray<FGatherTextFileExtension> Result;
 			Result.Add(FGatherTextFileExtension{TEXT("h")});
@@ -87,19 +87,19 @@ const TArray<FGatherTextFileExtension>& FGatherTextFromTextFilesConfiguration::G
 			Result.Add(FGatherTextFileExtension{TEXT("ini")});
 			return Result;
 		}();
-	return Result;
+	return DefaultTextFileExtensions;
 }
 
 const TArray<FGatherTextFileExtension>& FGatherTextFromPackagesConfiguration::GetDefaultPackageFileExtensions()
 {
-	static const TArray<FGatherTextFileExtension> Result = []()
+	static const TArray<FGatherTextFileExtension> DefaultPackageFileExtensions = []()
 		{
 			TArray<FGatherTextFileExtension> Result;
 			Result.Add(FGatherTextFileExtension{TEXT("umap")});
 			Result.Add(FGatherTextFileExtension{TEXT("uasset")});
 			return Result;
 		}();
-	return Result;
+	return DefaultPackageFileExtensions;
 }
 
 bool FGatherTextFromTextFilesConfiguration::Validate(const FString& RootDirectory, FText& OutError) const
