Translation editor fixes for loading translations of the native translations quickly. Also fixes import/export from the translation editor.

diff --git a/Engine/Source/Editor/Localization/Localization.Build.cs b/Engine/Source/Editor/Localization/Localization.Build.cs
--- a/Engine/Source/Editor/Localization/Localization.Build.cs
+++ b/Engine/Source/Editor/Localization/Localization.Build.cs
@@ -18,6 +18,13 @@ public class Localization : ModuleRules
 			new string[]
 			{
 				"Engine",
+				"Slate",
+				"SlateCore",
+				"InputCore",
+                "UnrealEd",
+                "EditorStyle",
+				"DesktopPlatform",
+                "SourceControl",
 			}
 		);
 
diff --git a/Engine/Source/Editor/Localization/Private/LocalizationCommandletExecution.cpp b/Engine/Source/Editor/Localization/Private/LocalizationCommandletExecution.cpp
new file mode 100644
--- /dev/null
+++ b/Engine/Source/Editor/Localization/Private/LocalizationCommandletExecution.cpp
@@ -0,0 +1,748 @@
+﻿// Copyright 1998-2015 Epic Games, Inc. All Rights Reserved.
+
+#include "LocalizationPrivatePCH.h"
+#include "LocalizationCommandletExecution.h"
+#include "LocalizationTargetTypes.h"
+#include "LocalizationConfigurationScript.h"
+#include "ISourceControlModule.h"
+#include "DesktopPlatformModule.h"
+#include "SThrobber.h"
+#include "Commandlets/CommandletHelpers.h"
+#include "EditorStyleSet.h"
+#include "UnrealEdMisc.h"
+#include "SourceControlHelpers.h"
+
+#define LOCTEXT_NAMESPACE "LocalizationCommandletExecutor"
+
+namespace
+{
+	class SLocalizationCommandletExecutor : public SCompoundWidget
+	{
+		SLATE_BEGIN_ARGS(SLocalizationCommandletExecutor) {}
+		SLATE_END_ARGS()
+
+	private:
+		struct FTaskListModel
+		{
+			enum class EState
+			{
+				Queued,
+				InProgress,
+				Failed,
+				Succeeded
+			};
+
+			FTaskListModel()
+				: State(EState::Queued)
+			{
+			}
+
+			LocalizationCommandletExecution::FTask Task;
+			EState State;
+			FString LogOutput;
+			FString ProcessArguments;
+		};
+
+		friend class STaskRow;
+
+	public:
+		SLocalizationCommandletExecutor();
+		~SLocalizationCommandletExecutor();
+
+		void Construct(const FArguments& Arguments, const TSharedRef<SWindow>& ParentWindow, const TArray<LocalizationCommandletExecution::FTask>& Tasks);
+		void Tick( const FGeometry& AllottedGeometry, const double InCurrentTime, const float InDeltaTime ) override;
+		bool WasSuccessful() const;
+		void Log(const FString& String);
+
+	private:
+		void ExecuteCommandlet(const TSharedRef<FTaskListModel>& TaskListModel);
+		void OnCommandletProcessCompletion(const int32 ReturnCode);
+		void CancelCommandlet();
+		void CleanUpProcessAndPump();
+
+		bool HasCompleted() const;
+
+		FText GetProgressMessageText() const;
+		TOptional<float> GetProgressPercentage() const;
+
+		TSharedRef<ITableRow> OnGenerateTaskListRow(TSharedPtr<FTaskListModel> TaskListModel, const TSharedRef<STableViewBase>& Table);
+		TSharedPtr<FTaskListModel> GetCurrentTaskToView() const;
+
+		FText GetCurrentTaskProcessArguments() const;
+		FText GetLogString() const;
+
+		FReply OnCopyLogClicked();
+		void CopyLogToClipboard();
+
+		FReply OnSaveLogClicked();
+
+		FText GetCloseButtonText() const;
+		FReply OnCloseButtonClicked();
+
+	private:
+		int32 CurrentTaskIndex;
+		TArray< TSharedPtr<FTaskListModel> > TaskListModels;
+		TSharedPtr<SProgressBar> ProgressBar;
+		TSharedPtr< SListView< TSharedPtr<FTaskListModel> > > TaskListView;
+
+		struct
+		{
+			FCriticalSection CriticalSection;
+			FString String;
+		} PendingLogData;
+
+
+		TSharedPtr<SWindow> ParentWindow;
+		TSharedPtr<FLocalizationCommandletProcess> CommandletProcess;
+		FRunnable* Runnable;
+		FRunnableThread* RunnableThread;
+	};
+
+	SLocalizationCommandletExecutor::SLocalizationCommandletExecutor()
+		: CurrentTaskIndex(INDEX_NONE)
+		, Runnable(nullptr)
+		, RunnableThread(nullptr)
+	{
+	}
+
+	SLocalizationCommandletExecutor::~SLocalizationCommandletExecutor()
+	{
+		CancelCommandlet();
+	}
+
+	void SLocalizationCommandletExecutor::Construct(const FArguments& Arguments, const TSharedRef<SWindow>& InParentWindow, const TArray<LocalizationCommandletExecution::FTask>& Tasks)
+	{
+		ParentWindow = InParentWindow;
+
+		for (const LocalizationCommandletExecution::FTask& Task : Tasks)
+		{
+			const TSharedRef<FTaskListModel> Model = MakeShareable(new FTaskListModel());
+			Model->Task = Task;
+			TaskListModels.Add(Model);
+		}
+
+		TSharedPtr<SScrollBar> VerticalScrollBar;
+		TSharedPtr<SScrollBar> HorizontalScrollBar;
+
+		ChildSlot
+			[
+				SNew(SVerticalBox)
+				+ SVerticalBox::Slot()
+				.AutoHeight()
+				.Padding(8.0, 16.0, 16.0, 0.0)
+				[
+					SNew(SVerticalBox)
+					+ SVerticalBox::Slot()
+					.AutoHeight()
+					[
+						SNew(STextBlock)
+						.Text(this, &SLocalizationCommandletExecutor::GetProgressMessageText)
+					]
+
+					+ SVerticalBox::Slot()
+						.AutoHeight()
+						.Padding(0.0, 4.0, 0.0, 0.0)
+						[
+							SAssignNew(ProgressBar, SProgressBar)
+							.Percent(this, &SLocalizationCommandletExecutor::GetProgressPercentage)
+						]
+				]
+				+ SVerticalBox::Slot()
+					.FillHeight(0.5)
+					.Padding(0.0, 32.0, 8.0, 0.0)
+					[
+						SNew(SBorder)
+						.BorderImage(FEditorStyle::GetBrush("ToolPanel.GroupBorder"))
+						.Padding(0.0f)
+						[
+							SAssignNew(TaskListView, SListView< TSharedPtr<FTaskListModel> >)
+							.HeaderRow
+							(
+							SNew(SHeaderRow)
+							+ SHeaderRow::Column("StatusIcon")
+							.DefaultLabel(FText::GetEmpty())
+							.FixedWidth(20.0)
+							+ SHeaderRow::Column("TaskName")
+							.DefaultLabel(LOCTEXT("TaskListNameColumnLabel", "Task"))
+							.FillWidth(1.0)
+							)
+							.ListItemsSource(&TaskListModels)
+							.OnGenerateRow(this, &SLocalizationCommandletExecutor::OnGenerateTaskListRow)
+							.ItemHeight(24.0)
+							.SelectionMode(ESelectionMode::Single)
+						]
+					]
+				+ SVerticalBox::Slot()
+					.FillHeight(0.5)
+					.Padding(0.0, 32.0, 8.0, 0.0)
+					[
+						SNew(SBorder)
+						.BorderImage(FEditorStyle::GetBrush("ToolPanel.GroupBorder"))
+						.Padding(0.0f)
+						[
+							SNew(SHorizontalBox)
+							+SHorizontalBox::Slot()
+							.FillWidth(1.0f)
+							[
+								SNew(SVerticalBox)
+								+SVerticalBox::Slot()
+								.FillHeight(1.0f)
+								[
+									SNew(SMultiLineEditableText)
+									.TextStyle(FEditorStyle::Get(), "LocalizationDashboard.CommandletLog.Text")
+									.Text(this, &SLocalizationCommandletExecutor::GetLogString)
+									.IsReadOnly(true)
+									.HScrollBar(HorizontalScrollBar)
+									.VScrollBar(VerticalScrollBar)
+								]
+								+SVerticalBox::Slot()
+									.AutoHeight()
+									[
+										SAssignNew(HorizontalScrollBar, SScrollBar)
+										.Orientation(EOrientation::Orient_Horizontal)
+									]
+							]
+							+SHorizontalBox::Slot()
+								.AutoWidth()
+								[
+									SAssignNew(VerticalScrollBar, SScrollBar)
+									.Orientation(EOrientation::Orient_Vertical)
+								]
+						]
+					]
+				+ SVerticalBox::Slot()
+					.AutoHeight()
+					.Padding(0.0f, 5.0f, 0.0f, 0.0f)
+					[
+						SNew(SHorizontalBox)
+						+ SHorizontalBox::Slot()
+						.AutoWidth()
+						[
+							SNew(SButton)
+							.ContentPadding(FMargin(6.0f, 2.0f))
+							.Text(LOCTEXT("CopyLogButtonText", "Copy Log"))
+							.ToolTipText(LOCTEXT("CopyLogButtonTooltip", "Copy the logged text to the clipboard."))
+							.OnClicked(this, &SLocalizationCommandletExecutor::OnCopyLogClicked)
+						]
+						+ SHorizontalBox::Slot()
+							.AutoWidth()
+							[
+								SNew(SButton)
+								.ContentPadding(FMargin(6.0f, 2.0f))
+								.IsEnabled(false)
+								.Text(LOCTEXT("SaveLogButtonText", "Save Log..."))
+								.ToolTipText(LOCTEXT("SaveLogButtonToolTip", "Save the logged text to a file."))
+								.OnClicked(this, &SLocalizationCommandletExecutor::OnSaveLogClicked)
+							]
+						+ SHorizontalBox::Slot()
+							.AutoWidth()
+							[
+								SNew(SButton)
+								.ContentPadding(FMargin(6.0f, 2.0f))
+								.OnClicked(this, &SLocalizationCommandletExecutor::OnCloseButtonClicked)
+								[
+									SNew(STextBlock)
+									.Text(this, &SLocalizationCommandletExecutor::GetCloseButtonText)
+								]
+							]
+					]
+			];
+
+		if(TaskListModels.Num() > 0)
+		{
+			CurrentTaskIndex = 0;
+			ExecuteCommandlet(TaskListModels[CurrentTaskIndex].ToSharedRef());
+		}
+	}
+
+	void SLocalizationCommandletExecutor::Tick( const FGeometry& AllottedGeometry, const double InCurrentTime, const float InDeltaTime )
+	{
+		SCompoundWidget::Tick(AllottedGeometry, InCurrentTime, InDeltaTime);
+
+		// Poll for log output data.
+		if (!PendingLogData.String.IsEmpty())
+		{
+			FString String;
+
+			// Copy the pending data string to the local string
+			{
+				FScopeLock ScopeLock(&PendingLogData.CriticalSection);
+				String = PendingLogData.String;
+				PendingLogData.String.Empty();
+			}
+
+			// Forward string to proper log.
+			const TSharedPtr<FTaskListModel> CurrentTaskModel = TaskListModels[CurrentTaskIndex];
+			CurrentTaskModel->LogOutput.Append(String);
+		}
+
+		// On Task Completed.
+		if (CommandletProcess.IsValid())
+		{
+			FProcHandle CurrentProcessHandle = CommandletProcess->GetHandle();
+			int32 ReturnCode;
+			if (CurrentProcessHandle.IsValid() && FPlatformProcess::GetProcReturnCode(CurrentProcessHandle, &ReturnCode))
+			{
+				OnCommandletProcessCompletion(ReturnCode);
+			}
+		}
+	}
+
+	bool SLocalizationCommandletExecutor::WasSuccessful() const
+	{
+		return HasCompleted() && !TaskListModels.ContainsByPredicate([](const TSharedPtr<FTaskListModel>& TaskListModel){return TaskListModel->State != FTaskListModel::EState::Succeeded;});
+	}
+
+	void SLocalizationCommandletExecutor::Log(const FString& String)
+	{
+		FScopeLock ScopeLock(&PendingLogData.CriticalSection);
+		PendingLogData.String += String;
+	}
+
+	void SLocalizationCommandletExecutor::OnCommandletProcessCompletion(const int32 ReturnCode)
+	{
+		CleanUpProcessAndPump();
+
+		// Handle return code.
+		TSharedPtr<FTaskListModel> CurrentTaskModel = TaskListModels[CurrentTaskIndex];
+
+		// Restore engine's source control settings if necessary.
+		if (!CurrentTaskModel->Task.ShouldUseProjectFile)
+		{
+			const FString& EngineIniFile = SourceControlHelpers::GetGlobalSettingsIni();
+			const FString BackupEngineIniFile = FPaths::EngineSavedDir() / FPaths::GetCleanFilename(EngineIniFile) + TEXT(".bak");
+			if(!IFileManager::Get().Move(*EngineIniFile, *BackupEngineIniFile))
+			{
+				// TODO: Error failed to restore engine source control settings INI.
+			}
+		}
+
+		// Zero code is successful.
+		if (ReturnCode == 0)
+		{
+			CurrentTaskModel->State = FTaskListModel::EState::Succeeded;
+
+			++CurrentTaskIndex;
+
+			// Begin new task if possible.
+			if (TaskListModels.IsValidIndex(CurrentTaskIndex))
+			{
+				CurrentTaskModel = TaskListModels[CurrentTaskIndex];
+				ExecuteCommandlet(CurrentTaskModel.ToSharedRef());
+			}
+		}
+		// Non-zero is a failure.
+		else
+		{
+			CurrentTaskModel->State = FTaskListModel::EState::Failed;
+		}
+	}
+
+	void SLocalizationCommandletExecutor::ExecuteCommandlet(const TSharedRef<FTaskListModel>& TaskListModel)
+	{
+		// Handle source control settings if not using project file for commandlet executable process.
+		if (!TaskListModel->Task.ShouldUseProjectFile)
+		{
+			const FString& EngineIniFile = SourceControlHelpers::GetGlobalSettingsIni();
+
+			// Backup engine's source control settings.
+			const FString BackupEngineIniFile = FPaths::EngineSavedDir() / FPaths::GetCleanFilename(EngineIniFile) + TEXT(".bak");
+			if (COPY_OK == IFileManager::Get().Copy(*BackupEngineIniFile, *EngineIniFile))
+			{
+				// Replace engine's source control settings with project's.
+				const FString& ProjectIniFile = SourceControlHelpers::GetSettingsIni();
+				if (COPY_OK == IFileManager::Get().Copy(*EngineIniFile, *ProjectIniFile))
+				{
+					// TODO: Error failed to overwrite engine source control settings INI.
+				}
+			}
+			else
+			{
+				// TODO: Error failed to backup engine source control settings INI.
+			}
+		}
+
+		CommandletProcess = FLocalizationCommandletProcess::Execute(TaskListModel->Task.ScriptPath, TaskListModel->Task.ShouldUseProjectFile);
+		if (CommandletProcess.IsValid())
+		{
+			TaskListModel->State = FTaskListModel::EState::InProgress;
+			TaskListModel->ProcessArguments = CommandletProcess->GetProcessArguments();
+		}
+		else
+		{
+			TaskListModel->State = FTaskListModel::EState::Failed;
+			return;
+		}
+
+		class FCommandletLogPump : public FRunnable
+		{
+		public:
+			FCommandletLogPump(void* const InReadPipe, const FProcHandle& InCommandletProcessHandle, SLocalizationCommandletExecutor& InCommandletWidget)
+				: ReadPipe(InReadPipe)
+				, CommandletProcessHandle(InCommandletProcessHandle)
+				, CommandletWidget(&InCommandletWidget)
+			{
+			}
+
+			uint32 Run() override
+			{
+				for(;;)
+				{
+					// Read from pipe.
+					const FString PipeString = FPlatformProcess::ReadPipe(ReadPipe);
+
+					// Process buffer.
+					if (!PipeString.IsEmpty())
+					{
+						// Add strings to log.
+						if (CommandletWidget)
+						{
+							CommandletWidget->Log(PipeString);
+						}
+					}
+
+					// If the process isn't running and there's no data in the pipe, we're done.
+					if (!FPlatformProcess::IsProcRunning(CommandletProcessHandle) && PipeString.IsEmpty())
+					{
+						break;
+					}
+
+					// Sleep.
+					FPlatformProcess::Sleep(0.0f);
+				}
+
+				int32 ReturnCode = 0;
+				return FPlatformProcess::GetProcReturnCode(CommandletProcessHandle, &ReturnCode) ? ReturnCode : -1;
+			}
+
+		private:
+			void* const ReadPipe;
+			FProcHandle CommandletProcessHandle;
+			SLocalizationCommandletExecutor* const CommandletWidget;
+		};
+
+		// Launch runnable thread.
+		Runnable = new FCommandletLogPump(CommandletProcess->GetReadPipe(), CommandletProcess->GetHandle(), *this);
+		RunnableThread = FRunnableThread::Create(Runnable, TEXT("Localization Commandlet Log Pump Thread"));
+	}
+
+	void SLocalizationCommandletExecutor::CancelCommandlet()
+	{
+		CleanUpProcessAndPump();
+	}
+
+	void SLocalizationCommandletExecutor::CleanUpProcessAndPump()
+	{
+		if (CommandletProcess.IsValid())
+		{
+			FProcHandle CommandletProcessHandle = CommandletProcess->GetHandle();
+			if (CommandletProcessHandle.IsValid() && FPlatformProcess::IsProcRunning(CommandletProcessHandle))
+			{
+				FPlatformProcess::TerminateProc(CommandletProcessHandle, true);
+			}
+			CommandletProcess.Reset();
+		}
+
+		if (RunnableThread)
+		{
+			RunnableThread->WaitForCompletion();
+			delete RunnableThread;
+			RunnableThread = nullptr;
+		}
+
+		if (Runnable)
+		{
+			delete Runnable;
+			Runnable = nullptr;
+		}
+
+	}
+
+	bool SLocalizationCommandletExecutor::HasCompleted() const
+	{
+		return CurrentTaskIndex == TaskListModels.Num();
+	}
+
+	FText SLocalizationCommandletExecutor::GetProgressMessageText() const
+	{
+		return TaskListModels.IsValidIndex(CurrentTaskIndex) ? TaskListModels[CurrentTaskIndex]->Task.Name : FText::GetEmpty();
+	}
+
+	TOptional<float> SLocalizationCommandletExecutor::GetProgressPercentage() const
+	{
+		return TOptional<float>(float(CurrentTaskIndex) / float(TaskListModels.Num()));
+	}
+
+	class STaskRow : public SMultiColumnTableRow< TSharedPtr<SLocalizationCommandletExecutor::FTaskListModel> >
+	{
+	public:
+		void Construct(const FTableRowArgs& InArgs, const TSharedRef<STableViewBase>& OwnerTableView, const TSharedRef<SLocalizationCommandletExecutor::FTaskListModel>& InTaskListModel);
+		TSharedRef<SWidget> GenerateWidgetForColumn(const FName& ColumnName);
+
+	private:
+		FSlateColor HandleIconColorAndOpacity() const;
+		const FSlateBrush* HandleIconImage() const;
+		EVisibility HandleThrobberVisibility() const;
+
+	private:
+		TSharedPtr<SLocalizationCommandletExecutor::FTaskListModel> TaskListModel;
+	};
+
+	void STaskRow::Construct(const FTableRowArgs& InArgs, const TSharedRef<STableViewBase>& OwnerTableView, const TSharedRef<SLocalizationCommandletExecutor::FTaskListModel>& InTaskListModel)
+	{
+		TaskListModel = InTaskListModel;
+
+		FSuperRowType::Construct(InArgs, OwnerTableView);
+	}
+
+	TSharedRef<SWidget> STaskRow::GenerateWidgetForColumn(const FName& ColumnName)
+	{
+		if (ColumnName == "StatusIcon")
+		{
+			return SNew(SOverlay)
+
+				+ SOverlay::Slot()
+				.HAlign(HAlign_Center)
+				.VAlign(VAlign_Center)
+				[
+					SNew(SThrobber)
+					.Animate(SThrobber::VerticalAndOpacity)
+					.NumPieces(1)
+					.Visibility(this, &STaskRow::HandleThrobberVisibility)
+				]
+			+ SOverlay::Slot()
+				.HAlign(HAlign_Center)
+				.VAlign(VAlign_Center)
+				[
+					SNew(SImage)
+					.ColorAndOpacity(this, &STaskRow::HandleIconColorAndOpacity)
+					.Image(this, &STaskRow::HandleIconImage)
+				];
+		}
+		else if (ColumnName == "TaskName")
+		{
+			return SNew(STextBlock)
+				.Text(TaskListModel->Task.Name)
+				.ToolTipText_Lambda( [this]{ return FText::FromString(TaskListModel->ProcessArguments); } );
+		}
+		else
+		{
+			return SNullWidget::NullWidget;
+		}
+	}
+
+	FSlateColor STaskRow::HandleIconColorAndOpacity( ) const
+	{
+		if (TaskListModel.IsValid())
+		{
+			switch(TaskListModel->State)
+			{
+			case SLocalizationCommandletExecutor::FTaskListModel::EState::InProgress:
+				return FLinearColor::Yellow;
+				break;
+			case SLocalizationCommandletExecutor::FTaskListModel::EState::Succeeded:
+				return FLinearColor::Green;
+				break;
+			case SLocalizationCommandletExecutor::FTaskListModel::EState::Failed:
+				return FLinearColor::Red;
+				break;
+			}
+		}
+
+		return FSlateColor::UseForeground();
+	}
+
+	const FSlateBrush* STaskRow::HandleIconImage( ) const
+	{
+		if (TaskListModel.IsValid())
+		{
+			switch(TaskListModel->State)
+			{
+			case SLocalizationCommandletExecutor::FTaskListModel::EState::Succeeded:
+				return FEditorStyle::GetBrush("Symbols.Check");
+				break;
+			case SLocalizationCommandletExecutor::FTaskListModel::EState::Failed:
+				return FEditorStyle::GetBrush("Icons.Cross");
+				break;
+			}
+		}
+
+		return NULL;
+	}
+
+	EVisibility STaskRow::HandleThrobberVisibility( ) const
+	{
+		if (TaskListModel.IsValid())
+		{
+			switch(TaskListModel->State)
+			{
+			case SLocalizationCommandletExecutor::FTaskListModel::EState::InProgress:
+				return EVisibility::Visible;
+				break;
+			}
+		}
+
+		return EVisibility::Hidden;
+	}
+
+	TSharedRef<ITableRow> SLocalizationCommandletExecutor::OnGenerateTaskListRow(TSharedPtr<FTaskListModel> TaskListModel, const TSharedRef<STableViewBase>& Table)
+	{
+		return SNew(STaskRow, Table, TaskListModel.ToSharedRef());
+	}
+
+	TSharedPtr<SLocalizationCommandletExecutor::FTaskListModel> SLocalizationCommandletExecutor::GetCurrentTaskToView() const
+	{
+		if (TaskListView.IsValid())
+		{
+			const TArray< TSharedPtr<FTaskListModel> > Selection = TaskListView->GetSelectedItems();
+			return Selection.Num() > 0 ? Selection.Top() : nullptr;
+		}
+		return nullptr;
+	}
+
+	FText SLocalizationCommandletExecutor::GetCurrentTaskProcessArguments() const
+	{
+		const TSharedPtr<SLocalizationCommandletExecutor::FTaskListModel> TaskToView = GetCurrentTaskToView();
+		return TaskToView.IsValid() ? FText::FromString(TaskToView->ProcessArguments) : FText::GetEmpty();
+	}
+
+	FText SLocalizationCommandletExecutor::GetLogString() const
+	{
+		const TSharedPtr<SLocalizationCommandletExecutor::FTaskListModel> TaskToView = GetCurrentTaskToView();
+		return TaskToView.IsValid() ? FText::FromString(TaskToView->LogOutput) : FText::GetEmpty();
+	}
+
+	FReply SLocalizationCommandletExecutor::OnCopyLogClicked()
+	{
+		CopyLogToClipboard();
+
+		return FReply::Handled();
+	}
+
+	void SLocalizationCommandletExecutor::CopyLogToClipboard()
+	{
+		FPlatformMisc::ClipboardCopy(*(GetLogString().ToString()));
+	}
+
+	FReply SLocalizationCommandletExecutor::OnSaveLogClicked()
+	{
+		const FString TextFileDescription = LOCTEXT("TextFileDescription", "Text File").ToString();
+		const FString TextFileExtension = TEXT("txt");
+		const FString TextFileExtensionWildcard = FString::Printf(TEXT("*.%s"), *TextFileExtension);
+		const FString FileTypes = FString::Printf(TEXT("%s (%s)|%s"), *TextFileDescription, *TextFileExtensionWildcard, *TextFileExtensionWildcard);
+		const FString DefaultFilename = FString::Printf(TEXT("%s.%s"), TEXT("Log"), *TextFileExtension);
+		const FString DefaultPath = FPaths::GameSavedDir();
+
+		TArray<FString> SaveFilenames;
+		IDesktopPlatform* DesktopPlatform = FDesktopPlatformModule::Get();
+
+		// Prompt the user for the filename
+		if (DesktopPlatform)
+		{
+			void* ParentWindowWindowHandle = NULL;
+
+			const TSharedPtr<SWindow>& ParentWindowPtr = FSlateApplication::Get().FindWidgetWindow(AsShared());
+			if (ParentWindowPtr.IsValid() && ParentWindowPtr->GetNativeWindow().IsValid())
+			{
+				ParentWindowWindowHandle = ParentWindowPtr->GetNativeWindow()->GetOSWindowHandle();
+			}
+
+			if (DesktopPlatform->SaveFileDialog(
+				ParentWindowWindowHandle,
+				LOCTEXT("SaveLogDialogTitle", "Save Log to File").ToString(),
+				DefaultPath,
+				DefaultFilename,
+				FileTypes,
+				EFileDialogFlags::None,
+				SaveFilenames
+				))
+			{
+				// Save to file.
+				FFileHelper::SaveStringToFile( GetLogString().ToString(), *(SaveFilenames.Last()) );
+			}
+		}
+
+		return FReply::Handled();
+	}
+
+	FText SLocalizationCommandletExecutor::GetCloseButtonText() const
+	{
+		return HasCompleted() ? LOCTEXT("OkayButtonText", "Okay") : LOCTEXT("CancelButtonText", "Cancel");
+	}
+
+	FReply SLocalizationCommandletExecutor::OnCloseButtonClicked()
+	{
+		if (!HasCompleted())
+		{
+			CancelCommandlet();
+		}
+		ParentWindow->RequestDestroyWindow();
+		return FReply::Handled();
+	}
+}
+
+bool LocalizationCommandletExecution::Execute(const TSharedRef<SWindow>& ParentWindow, const FText& Title, const TArray<FTask>& Tasks)
+{
+	const TSharedRef<SWindow> CommandletWindow = SNew(SWindow)
+		.Title(Title)
+		.SupportsMinimize(false)
+		.AutoCenter(EAutoCenter::PreferredWorkArea)
+		.ClientSize(FVector2D(600,400))
+		.ActivateWhenFirstShown(true)
+		.FocusWhenFirstShown(true);
+	const TSharedRef<SLocalizationCommandletExecutor> CommandletExecutor = SNew(SLocalizationCommandletExecutor, CommandletWindow, Tasks);
+	CommandletWindow->SetContent(CommandletExecutor);
+
+	FSlateApplication::Get().AddModalWindow(CommandletWindow, ParentWindow, false);
+	return CommandletExecutor->WasSuccessful();
+}
+
+TSharedPtr<FLocalizationCommandletProcess> FLocalizationCommandletProcess::Execute(const FString& ConfigFilePath, const bool UseProjectFile)
+{
+	// Create pipes.
+	void* ReadPipe;
+	void* WritePipe;
+	if (!FPlatformProcess::CreatePipe(ReadPipe, WritePipe))
+	{
+		return nullptr;
+	}
+
+	// Create process.
+	FString CommandletArguments;
+
+	const FString ConfigFileRelativeToGameDir = LocalizationConfigurationScript::MakePathRelativeForCommandletProcess(ConfigFilePath, UseProjectFile);
+	CommandletArguments = FString::Printf( TEXT("-config=\"%s\""), *ConfigFileRelativeToGameDir );
+
+	ISourceControlModule& SourceControlModule = ISourceControlModule::Get();
+	if (SourceControlModule.IsEnabled())
+	{
+		if (SourceControlModule.GetProvider().IsAvailable())
+		{
+			CommandletArguments = FString::Printf( TEXT("%s -EnableSCC -DisableSCCSubmit"), *CommandletArguments );
+		}
+	}
+
+	const FString ProjectFilePath = FString::Printf(TEXT("\"%s\""), *FPaths::ConvertRelativePathToFull(FPaths::GetProjectFilePath()));
+	const FString ProcessArguments = CommandletHelpers::BuildCommandletProcessArguments(TEXT("GatherText"), UseProjectFile ? *ProjectFilePath : nullptr, *CommandletArguments);
+	FProcHandle CommandletProcessHandle = FPlatformProcess::CreateProc(*FUnrealEdMisc::Get().GetExecutableForCommandlets(), *ProcessArguments, true, true, true, nullptr, 0, nullptr, WritePipe);
+
+	// Close pipes if process failed.
+	if (!CommandletProcessHandle.IsValid())
+	{
+		FPlatformProcess::ClosePipe(ReadPipe, WritePipe);
+		return nullptr;
+	}
+
+	return MakeShareable(new FLocalizationCommandletProcess(ReadPipe, WritePipe, CommandletProcessHandle, ProcessArguments));
+}
+
+FLocalizationCommandletProcess::~FLocalizationCommandletProcess()
+{
+	if (ProcessHandle.IsValid() && FPlatformProcess::IsProcRunning(ProcessHandle))
+	{
+		FPlatformProcess::TerminateProc(ProcessHandle);
+	}
+	FPlatformProcess::ClosePipe(ReadPipe, WritePipe);
+}
+
+#undef LOCTEXT_NAMESPACE
\ No newline at end of file
diff --git a/Engine/Source/Editor/Localization/Private/LocalizationCommandletTasks.cpp b/Engine/Source/Editor/Localization/Private/LocalizationCommandletTasks.cpp
new file mode 100644
--- /dev/null
+++ b/Engine/Source/Editor/Localization/Private/LocalizationCommandletTasks.cpp
@@ -0,0 +1,288 @@
+// Copyright 1998-2015 Epic Games, Inc. All Rights Reserved.
+
+#include "LocalizationPrivatePCH.h"
+#include "LocalizationCommandletTasks.h"
+#include "LocalizationCommandletExecution.h"
+#include "LocalizationConfigurationScript.h"
+
+#define LOCTEXT_NAMESPACE "LocalizationCommandletTasks"
+
+bool LocalizationCommandletTasks::GatherTargets(const TSharedRef<SWindow>& ParentWindow, const TArray<ULocalizationTarget*>& Targets)
+{
+	TArray<LocalizationCommandletExecution::FTask> Tasks;
+
+	for (ULocalizationTarget* Target : Targets)
+	{
+		const bool ShouldUseProjectFile = !Target->IsMemberOfEngineTargetSet();
+
+		FFormatNamedArguments Arguments;
+		Arguments.Add(TEXT("TargetName"), FText::FromString(Target->Settings.Name));
+
+		const FText GatherTaskName = FText::Format(LOCTEXT("GatherTaskNameFormat", "Gather Text for {TargetName}"), Arguments);
+		const FString GatherScriptPath = LocalizationConfigurationScript::GetGatherScriptPath(Target);
+		LocalizationConfigurationScript::GenerateGatherScript(Target).Write(GatherScriptPath);
+		Tasks.Add(LocalizationCommandletExecution::FTask(GatherTaskName, GatherScriptPath, ShouldUseProjectFile));
+	}
+
+	return LocalizationCommandletExecution::Execute(ParentWindow, LOCTEXT("GatherAllTargetsWindowTitle", "Gather Text for All Targets"), Tasks);
+}
+
+bool LocalizationCommandletTasks::GatherTarget(const TSharedRef<SWindow>& ParentWindow, ULocalizationTarget* const Target)
+{
+	TArray<LocalizationCommandletExecution::FTask> Tasks;
+	const bool ShouldUseProjectFile = !Target->IsMemberOfEngineTargetSet();
+
+	const FString GatherScriptPath = LocalizationConfigurationScript::GetGatherScriptPath(Target);
+	LocalizationConfigurationScript::GenerateGatherScript(Target).Write(GatherScriptPath);
+	Tasks.Add(LocalizationCommandletExecution::FTask(LOCTEXT("GatherTaskName", "Gather Text"), GatherScriptPath, ShouldUseProjectFile));
+
+	FFormatNamedArguments Arguments;
+	Arguments.Add(TEXT("TargetName"), FText::FromString(Target->Settings.Name));
+
+	const FText WindowTitle = FText::Format(LOCTEXT("GatherTargetWindowTitle", "Gather Text for Target {TargetName}"), Arguments);
+	return LocalizationCommandletExecution::Execute(ParentWindow, WindowTitle, Tasks);
+}
+
+
+bool LocalizationCommandletTasks::ImportTargets(const TSharedRef<SWindow>& ParentWindow, const TArray<ULocalizationTarget*>& Targets, const TOptional<FString> DirectoryPath)
+{
+	TArray<LocalizationCommandletExecution::FTask> Tasks;
+
+	for (ULocalizationTarget* Target : Targets)
+	{
+		const bool ShouldUseProjectFile = !Target->IsMemberOfEngineTargetSet();
+
+		FFormatNamedArguments Arguments;
+		Arguments.Add(TEXT("TargetName"), FText::FromString(Target->Settings.Name));
+
+		const FText ImportTaskName = FText::Format(LOCTEXT("ImportTaskNameFormat", "Import Translations for {TargetName}"), Arguments);
+		const FString ImportScriptPath = LocalizationConfigurationScript::GetImportScriptPath(Target, TOptional<FString>());
+		const TOptional<FString> DirectoryPathForTarget = DirectoryPath.IsSet() ? DirectoryPath.GetValue() / Target->Settings.Name : TOptional<FString>();
+		LocalizationConfigurationScript::GenerateImportScript(Target, TOptional<FString>(), DirectoryPathForTarget).Write(ImportScriptPath);
+		Tasks.Add(LocalizationCommandletExecution::FTask(ImportTaskName, ImportScriptPath, ShouldUseProjectFile));
+
+		const FText ReportTaskName = FText::Format(LOCTEXT("ReportTaskNameFormat", "Generate Reports for {TargetName}"), Arguments);
+		const FString ReportScriptPath = LocalizationConfigurationScript::GetWordCountReportScriptPath(Target);
+		LocalizationConfigurationScript::GenerateWordCountReportScript(Target).Write(ReportScriptPath);
+		Tasks.Add(LocalizationCommandletExecution::FTask(ReportTaskName, ReportScriptPath, ShouldUseProjectFile));
+	}
+
+	return LocalizationCommandletExecution::Execute(ParentWindow, LOCTEXT("ImportForAllTargetsWindowTitle", "Import Translations for All Targets"), Tasks);
+}
+
+bool LocalizationCommandletTasks::ImportTarget(const TSharedRef<SWindow>& ParentWindow, ULocalizationTarget* const Target, const TOptional<FString> DirectoryPath)
+{
+	TArray<LocalizationCommandletExecution::FTask> Tasks;
+	const bool ShouldUseProjectFile = !Target->IsMemberOfEngineTargetSet();
+
+	const FString ImportScriptPath = LocalizationConfigurationScript::GetImportScriptPath(Target, TOptional<FString>());
+	LocalizationConfigurationScript::GenerateImportScript(Target, TOptional<FString>(), DirectoryPath).Write(ImportScriptPath);
+	Tasks.Add(LocalizationCommandletExecution::FTask(LOCTEXT("ImportTaskName", "Import Translations"), ImportScriptPath, ShouldUseProjectFile));
+
+	const FString ReportScriptPath = LocalizationConfigurationScript::GetWordCountReportScriptPath(Target);
+	LocalizationConfigurationScript::GenerateWordCountReportScript(Target).Write(ReportScriptPath);
+	Tasks.Add(LocalizationCommandletExecution::FTask(LOCTEXT("ReportTaskName", "Generate Reports"), ReportScriptPath, ShouldUseProjectFile));
+
+	FFormatNamedArguments Arguments;
+	Arguments.Add(TEXT("TargetName"), FText::FromString(Target->Settings.Name));
+
+	const FText WindowTitle = FText::Format(LOCTEXT("ImportForTargetWindowTitle", "Import Translations for Target {TargetName}"), Arguments);
+	return LocalizationCommandletExecution::Execute(ParentWindow, WindowTitle, Tasks);
+}
+
+bool LocalizationCommandletTasks::ImportCulture(const TSharedRef<SWindow>& ParentWindow, ULocalizationTarget* const Target, const FString& CultureName, const TOptional<FString> FilePath)
+{
+	FCulturePtr Culture = FInternationalization::Get().GetCulture(CultureName);
+	if (!Culture.IsValid())
+	{
+		return false;
+	}
+
+	TArray<LocalizationCommandletExecution::FTask> Tasks;
+	const bool ShouldUseProjectFile = !Target->IsMemberOfEngineTargetSet();
+
+	const FString DefaultImportScriptPath = LocalizationConfigurationScript::GetImportScriptPath(Target, TOptional<FString>(CultureName));
+	const FString ImportScriptPath = FPaths::CreateTempFilename(*FPaths::GetPath(DefaultImportScriptPath), *FPaths::GetBaseFilename(DefaultImportScriptPath), *FPaths::GetExtension(DefaultImportScriptPath, true));
+	LocalizationConfigurationScript::GenerateImportScript(Target, TOptional<FString>(CultureName), FilePath).Write(ImportScriptPath);
+	Tasks.Add(LocalizationCommandletExecution::FTask(LOCTEXT("ImportTaskName", "Import Translations"), ImportScriptPath, ShouldUseProjectFile));
+
+	const FString ReportScriptPath = LocalizationConfigurationScript::GetWordCountReportScriptPath(Target);
+	LocalizationConfigurationScript::GenerateWordCountReportScript(Target).Write(ReportScriptPath);
+	Tasks.Add(LocalizationCommandletExecution::FTask(LOCTEXT("ReportTaskName", "Generate Reports"), ReportScriptPath, ShouldUseProjectFile));
+
+	FFormatNamedArguments Arguments;
+	Arguments.Add(TEXT("CultureName"), FText::FromString(Culture->GetDisplayName()));
+	Arguments.Add(TEXT("TargetName"), FText::FromString(Target->Settings.Name));
+
+	const FText WindowTitle = FText::Format(LOCTEXT("ImportCultureForTargetWindowTitle", "Import {CultureName} Translations for Target {TargetName}"), Arguments);
+
+	bool HasSucceeeded = LocalizationCommandletExecution::Execute(ParentWindow, WindowTitle, Tasks);
+	IFileManager::Get().Delete(*ImportScriptPath); // Don't clutter up the loc config directory with scripts for individual cultures.
+	return HasSucceeeded;
+
+}
+
+bool LocalizationCommandletTasks::ExportTargets(const TSharedRef<SWindow>& ParentWindow, const TArray<ULocalizationTarget*>& Targets, const TOptional<FString> DirectoryPath)
+{
+	TArray<LocalizationCommandletExecution::FTask> Tasks;
+
+	for (ULocalizationTarget* Target : Targets)
+	{
+		const bool ShouldUseProjectFile = !Target->IsMemberOfEngineTargetSet();
+
+		FFormatNamedArguments Arguments;
+		Arguments.Add(TEXT("TargetName"), FText::FromString(Target->Settings.Name));
+
+		const FText ExportTaskName = FText::Format(LOCTEXT("ExportTaskNameFormat", "Export Translations for {TargetName}"), Arguments);
+		const FString ExportScriptPath = LocalizationConfigurationScript::GetExportScriptPath(Target, TOptional<FString>());
+		const TOptional<FString> DirectoryPathForTarget = DirectoryPath.IsSet() ? DirectoryPath.GetValue() / Target->Settings.Name : TOptional<FString>();
+		LocalizationConfigurationScript::GenerateExportScript(Target, TOptional<FString>(), DirectoryPathForTarget).Write(ExportScriptPath);
+		Tasks.Add(LocalizationCommandletExecution::FTask(ExportTaskName, ExportScriptPath, ShouldUseProjectFile));
+	}
+
+	return LocalizationCommandletExecution::Execute(ParentWindow, LOCTEXT("ExportForAllTargetsWindowTitle", "Export Translations for All Targets"), Tasks);
+}
+
+bool LocalizationCommandletTasks::ExportTarget(const TSharedRef<SWindow>& ParentWindow, ULocalizationTarget* const Target, const TOptional<FString> DirectoryPath)
+{
+	TArray<LocalizationCommandletExecution::FTask> Tasks;
+	const bool ShouldUseProjectFile = !Target->IsMemberOfEngineTargetSet();
+
+	const FString ExportScriptPath = LocalizationConfigurationScript::GetExportScriptPath(Target, TOptional<FString>());
+	LocalizationConfigurationScript::GenerateExportScript(Target, TOptional<FString>(), DirectoryPath).Write(ExportScriptPath);
+	Tasks.Add(LocalizationCommandletExecution::FTask(LOCTEXT("ExportTaskName", "Export Translations"), ExportScriptPath, ShouldUseProjectFile));
+
+	FFormatNamedArguments Arguments;
+	Arguments.Add(TEXT("TargetName"), FText::FromString(Target->Settings.Name));
+
+	const FText WindowTitle = FText::Format(LOCTEXT("ExportForTargetWindowTitle", "Export Translations for Target {TargetName}"), Arguments);
+	return LocalizationCommandletExecution::Execute(ParentWindow, WindowTitle, Tasks);
+}
+
+bool LocalizationCommandletTasks::ExportCulture(const TSharedRef<SWindow>& ParentWindow, ULocalizationTarget* const Target, const FString& CultureName, const TOptional<FString> FilePath)
+{
+	FCulturePtr Culture = FInternationalization::Get().GetCulture(CultureName);
+	if (!Culture.IsValid())
+	{
+		return false;
+	}
+
+	TArray<LocalizationCommandletExecution::FTask> Tasks;
+	const bool ShouldUseProjectFile = !Target->IsMemberOfEngineTargetSet();
+
+	const FString DefaultExportScriptPath = LocalizationConfigurationScript::GetExportScriptPath(Target, TOptional<FString>(CultureName));
+	const FString ExportScriptPath = FPaths::CreateTempFilename(*FPaths::GetPath(DefaultExportScriptPath), *FPaths::GetBaseFilename(DefaultExportScriptPath), *FPaths::GetExtension(DefaultExportScriptPath, true));
+	LocalizationConfigurationScript::GenerateExportScript(Target, TOptional<FString>(CultureName), FilePath).Write(ExportScriptPath);
+	Tasks.Add(LocalizationCommandletExecution::FTask(LOCTEXT("ExportTaskName", "Export Translations"), ExportScriptPath, ShouldUseProjectFile));
+
+	FFormatNamedArguments Arguments;
+	Arguments.Add(TEXT("CultureName"), FText::FromString(Culture->GetDisplayName()));
+	Arguments.Add(TEXT("TargetName"), FText::FromString(Target->Settings.Name));
+
+	const FText WindowTitle = FText::Format(LOCTEXT("ExportCultureForTargetWindowTitle", "Export {CultureName} Translations for Target {TargetName}"), Arguments);
+
+	bool HasSucceeeded = LocalizationCommandletExecution::Execute(ParentWindow, WindowTitle, Tasks);
+	IFileManager::Get().Delete(*ExportScriptPath); // Don't clutter up the loc config directory with scripts for individual cultures.
+	return HasSucceeeded;
+}
+
+bool LocalizationCommandletTasks::GenerateWordCountReportsForTargets(const TSharedRef<SWindow>& ParentWindow, const TArray<ULocalizationTarget*>& Targets)
+{
+	TArray<LocalizationCommandletExecution::FTask> Tasks;
+
+	for (ULocalizationTarget* Target : Targets)
+	{
+		const bool ShouldUseProjectFile = !Target->IsMemberOfEngineTargetSet();
+
+		FFormatNamedArguments Arguments;
+		Arguments.Add(TEXT("TargetName"), FText::FromString(Target->Settings.Name));
+
+		const FText ReportTaskName = FText::Format(LOCTEXT("ReportTaskNameFormat", "Generate Word Count Report for {TargetName}"), Arguments);
+		const FString ReportScriptPath = LocalizationConfigurationScript::GetWordCountReportScriptPath(Target);
+		LocalizationConfigurationScript::GenerateWordCountReportScript(Target).Write(ReportScriptPath);
+		Tasks.Add(LocalizationCommandletExecution::FTask(ReportTaskName, ReportScriptPath, ShouldUseProjectFile));
+	}
+
+	return LocalizationCommandletExecution::Execute(ParentWindow, LOCTEXT("GenerateReportsForAllTargetsWindowTitle", "Generate Word Count Reports for All Targets"), Tasks);
+}
+
+bool LocalizationCommandletTasks::GenerateWordCountReportForTarget(const TSharedRef<SWindow>& ParentWindow, ULocalizationTarget* const Target)
+{
+	TArray<LocalizationCommandletExecution::FTask> Tasks;
+	const bool ShouldUseProjectFile = !Target->IsMemberOfEngineTargetSet();
+
+	const FString ReportScriptPath = LocalizationConfigurationScript::GetWordCountReportScriptPath(Target);
+	LocalizationConfigurationScript::GenerateWordCountReportScript(Target).Write(ReportScriptPath);
+	Tasks.Add(LocalizationCommandletExecution::FTask(LOCTEXT("ReportTaskName", "Generate Word Count Report"), ReportScriptPath, ShouldUseProjectFile));
+
+	FFormatNamedArguments Arguments;
+	Arguments.Add(TEXT("TargetName"), FText::FromString(Target->Settings.Name));
+
+	const FText WindowTitle = FText::Format(LOCTEXT("GenerateReportForTargetWindowTitle", "Generate Word Count Report for Target {TargetName}"), Arguments);
+	return LocalizationCommandletExecution::Execute(ParentWindow, WindowTitle, Tasks);
+}
+
+bool LocalizationCommandletTasks::CompileTargets(const TSharedRef<SWindow>& ParentWindow, const TArray<ULocalizationTarget*>& Targets)
+{
+	TArray<LocalizationCommandletExecution::FTask> Tasks;
+
+	for (ULocalizationTarget* Target : Targets)
+	{
+		const bool ShouldUseProjectFile = !Target->IsMemberOfEngineTargetSet();
+
+		FFormatNamedArguments Arguments;
+		Arguments.Add(TEXT("TargetName"), FText::FromString(Target->Settings.Name));
+
+		const FText CompileTaskName = FText::Format(LOCTEXT("CompileTaskNameFormat", "Compile Translations for {TargetName}"), Arguments);
+		const FString CompileScriptPath = LocalizationConfigurationScript::GetCompileScriptPath(Target);
+		LocalizationConfigurationScript::GenerateCompileScript(Target).Write(CompileScriptPath);
+		Tasks.Add(LocalizationCommandletExecution::FTask(CompileTaskName, CompileScriptPath, ShouldUseProjectFile));
+	}
+
+	return LocalizationCommandletExecution::Execute(ParentWindow, LOCTEXT("GenerateLocResForAllTargetsWindowTitle", "Compile Translations for All Targets"), Tasks);
+}
+
+bool LocalizationCommandletTasks::CompileTarget(const TSharedRef<SWindow>& ParentWindow, ULocalizationTarget* const Target)
+{
+	TArray<LocalizationCommandletExecution::FTask> Tasks;
+	const bool ShouldUseProjectFile = !Target->IsMemberOfEngineTargetSet();
+
+	const FString CompileScriptPath = LocalizationConfigurationScript::GetCompileScriptPath(Target);
+	LocalizationConfigurationScript::GenerateCompileScript(Target).Write(CompileScriptPath);
+	Tasks.Add(LocalizationCommandletExecution::FTask(LOCTEXT("CompileTaskName", "Compile Translations"), CompileScriptPath, ShouldUseProjectFile));
+
+	FFormatNamedArguments Arguments;
+	Arguments.Add(TEXT("TargetName"), FText::FromString(Target->Settings.Name));
+
+	const FText WindowTitle = FText::Format(LOCTEXT("GenerateLocResForTargetWindowTitle", "Compile Translations for Target {TargetName}"), Arguments);
+	return LocalizationCommandletExecution::Execute(ParentWindow, WindowTitle, Tasks);
+}
+
+bool LocalizationCommandletTasks::CompileCulture(const TSharedRef<SWindow>& ParentWindow, ULocalizationTarget* const Target, const FString& CultureName)
+{
+	FCulturePtr Culture = FInternationalization::Get().GetCulture(CultureName);
+	if (!Culture.IsValid())
+	{
+		return false;
+	}
+
+	TArray<LocalizationCommandletExecution::FTask> Tasks;
+	const bool ShouldUseProjectFile = !Target->IsMemberOfEngineTargetSet();
+
+	const FString DefaultCompileScriptPath = LocalizationConfigurationScript::GetCompileScriptPath(Target, TOptional<FString>(CultureName));
+	const FString CompileScriptPath = FPaths::CreateTempFilename(*FPaths::GetPath(DefaultCompileScriptPath), *FPaths::GetBaseFilename(DefaultCompileScriptPath), *FPaths::GetExtension(DefaultCompileScriptPath, true));
+	LocalizationConfigurationScript::GenerateCompileScript(Target, TOptional<FString>(CultureName)).Write(CompileScriptPath);
+	Tasks.Add(LocalizationCommandletExecution::FTask(LOCTEXT("CompileTaskName", "Compile Translations"), CompileScriptPath, ShouldUseProjectFile));
+
+	FFormatNamedArguments Arguments;
+	Arguments.Add(TEXT("CultureName"), FText::FromString(Culture->GetDisplayName()));
+	Arguments.Add(TEXT("TargetName"), FText::FromString(Target->Settings.Name));
+
+	const FText WindowTitle = FText::Format(LOCTEXT("CompileCultureForTargetWindowTitle", "Compile {CultureName} Translations for Target {TargetName}"), Arguments);
+
+	bool HasSucceeeded = LocalizationCommandletExecution::Execute(ParentWindow, WindowTitle, Tasks);
+	IFileManager::Get().Delete(*CompileScriptPath); // Don't clutter up the loc config directory with scripts for individual cultures.
+	return HasSucceeeded;
+}
+
+#undef LOCTEXT_NAMESPACE
\ No newline at end of file
diff --git a/Engine/Source/Editor/Localization/Public/LocalizationCommandletExecution.h b/Engine/Source/Editor/Localization/Public/LocalizationCommandletExecution.h
new file mode 100644
--- /dev/null
+++ b/Engine/Source/Editor/Localization/Public/LocalizationCommandletExecution.h
@@ -0,0 +1,63 @@
+// Copyright 1998-2015 Epic Games, Inc. All Rights Reserved.
+#pragma once 
+
+#include "LocalizationConfigurationScript.h"
+
+namespace LocalizationCommandletExecution
+{
+	struct LOCALIZATION_API FTask
+	{
+		FTask() {}
+
+		FTask(const FText& InName, const FString& InScriptPath, const bool InShouldUseProjectFile)
+			: Name(InName)
+			, ScriptPath(InScriptPath)
+			, ShouldUseProjectFile(InShouldUseProjectFile)
+		{}
+
+		FText Name;
+		FString ScriptPath;
+		bool ShouldUseProjectFile;
+	};
+
+	bool Execute(const TSharedRef<SWindow>& ParentWindow, const FText& Title, const TArray<FTask>& Tasks);
+};
+
+class LOCALIZATION_API FLocalizationCommandletProcess : public TSharedFromThis<FLocalizationCommandletProcess>
+{
+public:
+	static TSharedPtr<FLocalizationCommandletProcess> Execute(const FString& ConfigFilePath, const bool UseProjectFile = true);
+
+private:
+	FLocalizationCommandletProcess(void* const InReadPipe, void* const InWritePipe, const FProcHandle InProcessHandle, const FString& InProcessArguments)
+		: ReadPipe(InReadPipe)
+		, WritePipe(InWritePipe)
+		, ProcessHandle(InProcessHandle)
+		, ProcessArguments(InProcessArguments)
+	{
+	}
+
+public:
+	~FLocalizationCommandletProcess();
+
+	void* GetReadPipe() const
+	{
+		return ReadPipe;
+	}
+
+	FProcHandle& GetHandle()
+	{
+		return ProcessHandle;
+	}
+
+	const FString& GetProcessArguments() const
+	{
+		return ProcessArguments;
+	}
+
+private:
+	void* const ReadPipe;
+	void* const WritePipe;
+	FProcHandle ProcessHandle;
+	FString ProcessArguments;
+};
\ No newline at end of file
diff --git a/Engine/Source/Editor/Localization/Public/LocalizationCommandletTasks.h b/Engine/Source/Editor/Localization/Public/LocalizationCommandletTasks.h
new file mode 100644
--- /dev/null
+++ b/Engine/Source/Editor/Localization/Public/LocalizationCommandletTasks.h
@@ -0,0 +1,25 @@
+// Copyright 1998-2015 Epic Games, Inc. All Rights Reserved.
+#pragma once
+
+#include "LocalizationTargetTypes.h"
+
+namespace LocalizationCommandletTasks
+{
+	LOCALIZATION_API bool GatherTargets(const TSharedRef<SWindow>& ParentWindow, const TArray<ULocalizationTarget*>& Targets);
+	LOCALIZATION_API bool GatherTarget(const TSharedRef<SWindow>& ParentWindow, ULocalizationTarget* const Target);
+
+	LOCALIZATION_API bool ImportTargets(const TSharedRef<SWindow>& ParentWindow, const TArray<ULocalizationTarget*>& Targets, const TOptional<FString> DirectoryPath = TOptional<FString>());
+	LOCALIZATION_API bool ImportTarget(const TSharedRef<SWindow>& ParentWindow, ULocalizationTarget* const Target, const TOptional<FString> DirectoryPath = TOptional<FString>());
+	LOCALIZATION_API bool ImportCulture(const TSharedRef<SWindow>& ParentWindow, ULocalizationTarget* const Target, const FString& CultureName, const TOptional<FString> FilePath = TOptional<FString>());
+
+	LOCALIZATION_API bool ExportTargets(const TSharedRef<SWindow>& ParentWindow, const TArray<ULocalizationTarget*>& Targets, const TOptional<FString> DirectoryPath = TOptional<FString>());
+	LOCALIZATION_API bool ExportTarget(const TSharedRef<SWindow>& ParentWindow, ULocalizationTarget* const Target, const TOptional<FString> DirectoryPath = TOptional<FString>());
+	LOCALIZATION_API bool ExportCulture(const TSharedRef<SWindow>& ParentWindow, ULocalizationTarget* const Target, const FString& CultureName, const TOptional<FString> FilePath = TOptional<FString>());
+
+	LOCALIZATION_API bool GenerateWordCountReportsForTargets(const TSharedRef<SWindow>& ParentWindow, const TArray<ULocalizationTarget*>& Targets);
+	LOCALIZATION_API bool GenerateWordCountReportForTarget(const TSharedRef<SWindow>& ParentWindow, ULocalizationTarget* const Target);
+
+	LOCALIZATION_API bool CompileTargets(const TSharedRef<SWindow>& ParentWindow, const TArray<ULocalizationTarget*>& Targets);
+	LOCALIZATION_API bool CompileTarget(const TSharedRef<SWindow>& ParentWindow, ULocalizationTarget* const Target);
+	LOCALIZATION_API bool CompileCulture(const TSharedRef<SWindow>& ParentWindow, ULocalizationTarget* const Target, const FString& CultureName);
+}
\ No newline at end of file
diff --git a/Engine/Source/Editor/LocalizationDashboard/Private/LocalizationCommandletExecution.cpp b/Engine/Source/Editor/LocalizationDashboard/Private/LocalizationCommandletExecution.cpp
deleted file mode 100644
diff --git a/Engine/Source/Editor/LocalizationDashboard/Private/LocalizationCommandletExecution.h b/Engine/Source/Editor/LocalizationDashboard/Private/LocalizationCommandletExecution.h
deleted file mode 100644
diff --git a/Engine/Source/Editor/LocalizationDashboard/Private/LocalizationCommandletTasks.cpp b/Engine/Source/Editor/LocalizationDashboard/Private/LocalizationCommandletTasks.cpp
deleted file mode 100644
diff --git a/Engine/Source/Editor/LocalizationDashboard/Private/LocalizationCommandletTasks.h b/Engine/Source/Editor/LocalizationDashboard/Private/LocalizationCommandletTasks.h
deleted file mode 100644
diff --git a/Engine/Source/Editor/LocalizationDashboard/Private/SLocalizationTargetEditorCultureRow.cpp b/Engine/Source/Editor/LocalizationDashboard/Private/SLocalizationTargetEditorCultureRow.cpp
--- a/Engine/Source/Editor/LocalizationDashboard/Private/SLocalizationTargetEditorCultureRow.cpp
+++ b/Engine/Source/Editor/LocalizationDashboard/Private/SLocalizationTargetEditorCultureRow.cpp
@@ -325,7 +325,13 @@ FReply SLocalizationTargetEditorCultureRow::Edit()
 	ULocalizationTarget* const LocalizationTarget = GetTarget();
 	if (Culture.IsValid() && LocalizationTarget)
 	{
-		FModuleManager::LoadModuleChecked<ITranslationEditor>("TranslationEditor").OpenTranslationEditor(LocalizationConfigurationScript::GetManifestPath(LocalizationTarget), LocalizationConfigurationScript::GetArchivePath(LocalizationTarget, Culture->GetName()), LocalizationTarget->Settings.Guid);
+		FString NativeCultureName;
+		if (LocalizationTarget->Settings.SupportedCulturesStatistics.IsValidIndex(LocalizationTarget->Settings.NativeCultureIndex))
+		{
+			NativeCultureName = LocalizationTarget->Settings.SupportedCulturesStatistics[LocalizationTarget->Settings.NativeCultureIndex].CultureName;
+		}
+		
+		FModuleManager::LoadModuleChecked<ITranslationEditor>("TranslationEditor").OpenTranslationEditor(LocalizationTarget, Culture->GetName());
 	}
 
 	return FReply::Handled();
diff --git a/Engine/Source/Editor/TranslationEditor/Private/ITranslationEditor.cpp b/Engine/Source/Editor/TranslationEditor/Private/ITranslationEditor.cpp
--- a/Engine/Source/Editor/TranslationEditor/Private/ITranslationEditor.cpp
+++ b/Engine/Source/Editor/TranslationEditor/Private/ITranslationEditor.cpp
@@ -4,18 +4,19 @@
 #include "TranslationEditorPrivatePCH.h"
 #include "ITranslationEditor.h"
 #include "Editor/TranslationEditor/Private/TranslationPickerWidget.h"
+#include "LocalizationConfigurationScript.h"
 
-/** To keep track of what translation editors are open editing which manifest files */
+/** To keep track of what translation editors are open editing which archive files */
 TMap<FString, ITranslationEditor*> ITranslationEditor::OpenTranslationEditors = TMap<FString, ITranslationEditor*>();
 
-void ITranslationEditor::OpenTranslationEditor(const FString& InManifestFile, const FString& InArchiveFile, const FGuid& LocalizationTargetGuid)
+void ITranslationEditor::OpenTranslationEditor(const FString& InManifestFile, const FString& InNativeArchiveFile, const FString& InArchiveFileToEdit)
 {
 	// Only open a new translation editor if there isn't another one open for this archive file already
-	if (!OpenTranslationEditors.Contains(InArchiveFile))
+	if (!OpenTranslationEditors.Contains(InArchiveFileToEdit))
 	{
 		FTranslationEditorModule& TranslationEditorModule = FModuleManager::LoadModuleChecked<FTranslationEditorModule>("TranslationEditor");
 		bool bLoadedSuccessfully = false;
-		ITranslationEditor* NewTranslationEditor = (ITranslationEditor*)&(TranslationEditorModule.CreateTranslationEditor(InManifestFile, InArchiveFile, LocalizationTargetGuid, bLoadedSuccessfully).Get());
+		ITranslationEditor* NewTranslationEditor = (ITranslationEditor*)&(TranslationEditorModule.CreateTranslationEditor(InManifestFile, InNativeArchiveFile, InArchiveFileToEdit, bLoadedSuccessfully).Get());
 
 		NewTranslationEditor->RegisterTranslationEditor();
 
@@ -27,7 +28,41 @@ void ITranslationEditor::OpenTranslationEditor(const FString& InManifestFile, co
 	else
 	{
 		// If already editing this archive file, flash the tab that contains the editor that has that file open
-		ITranslationEditor* OpenEditor = *OpenTranslationEditors.Find(InArchiveFile);
+		ITranslationEditor* OpenEditor = *OpenTranslationEditors.Find(InArchiveFileToEdit);
+		OpenEditor->FocusWindow();
+	}
+}
+
+void ITranslationEditor::OpenTranslationEditor(ULocalizationTarget* const LocalizationTarget, const FString& CultureToEdit)
+{
+	check(LocalizationTarget);
+	const FString ManifestFile = LocalizationConfigurationScript::GetManifestPath(LocalizationTarget);
+	FString NativeCultureName;
+	if (LocalizationTarget->Settings.SupportedCulturesStatistics.IsValidIndex(LocalizationTarget->Settings.NativeCultureIndex))
+	{
+		NativeCultureName = LocalizationTarget->Settings.SupportedCulturesStatistics[LocalizationTarget->Settings.NativeCultureIndex].CultureName;
+	}
+	const FString NativeArchiveFile = NativeCultureName.IsEmpty() ? FString() : LocalizationConfigurationScript::GetArchivePath(LocalizationTarget, NativeCultureName);
+	const FString ArchiveFile = LocalizationConfigurationScript::GetArchivePath(LocalizationTarget, CultureToEdit);
+
+	// Only open a new translation editor if there isn't another one open for this archive file already
+	if (!OpenTranslationEditors.Contains(ArchiveFile))
+	{
+		FTranslationEditorModule& TranslationEditorModule = FModuleManager::LoadModuleChecked<FTranslationEditorModule>("TranslationEditor");
+		bool bLoadedSuccessfully = false;
+		ITranslationEditor* NewTranslationEditor = (ITranslationEditor*)&(TranslationEditorModule.CreateTranslationEditor(LocalizationTarget, CultureToEdit, bLoadedSuccessfully).Get());
+
+		NewTranslationEditor->RegisterTranslationEditor();
+
+		if (!bLoadedSuccessfully)
+		{
+			NewTranslationEditor->CloseWindow();
+		}
+	}
+	else
+	{
+		// If already editing this archive file, flash the tab that contains the editor that has that file open
+		ITranslationEditor* OpenEditor = *OpenTranslationEditors.Find(ArchiveFile);
 		OpenEditor->FocusWindow();
 	}
 }
diff --git a/Engine/Source/Editor/TranslationEditor/Private/TranslationDataManager.cpp b/Engine/Source/Editor/TranslationEditor/Private/TranslationDataManager.cpp
--- a/Engine/Source/Editor/TranslationEditor/Private/TranslationDataManager.cpp
+++ b/Engine/Source/Editor/TranslationEditor/Private/TranslationDataManager.cpp
@@ -1,4 +1,4 @@
-// Copyright 1998-2015 Epic Games, Inc. All Rights Reserved.
+﻿// Copyright 1998-2015 Epic Games, Inc. All Rights Reserved.
 
 #include "TranslationEditorPrivatePCH.h"
 #include "TranslationEditor.h"
@@ -23,15 +23,44 @@ DEFINE_LOG_CATEGORY_STATIC(LogTranslationEditor, Log, All);
 #define LOCTEXT_NAMESPACE "TranslationDataManager"
 
 
-FTranslationDataManager::FTranslationDataManager( const FString& InManifestFilePath, const FString& InArchiveFilePath )
-: ManifestFilePath(InManifestFilePath)
-, ArchiveFilePath(InArchiveFilePath)
-, bLoadedSuccessfully(true)
+FTranslationDataManager::FTranslationDataManager( const FString& InManifestFilePath, const FString& InNativeArchiveFilePath, const FString& InArchiveFilePath )
+	: AssociatedLocalizationTarget(nullptr)
+	, OpenedManifestFilePath(InManifestFilePath)
+	, NativeArchiveFilePath(InNativeArchiveFilePath)
+	, OpenedArchiveFilePath(InArchiveFilePath)
+	, bLoadedSuccessfully(true)
+{
+	Initialize();
+}
+
+FTranslationDataManager::FTranslationDataManager(ULocalizationTarget* const LocalizationTarget, const FString& CultureToEdit)
+	: AssociatedLocalizationTarget(LocalizationTarget)
+	, bLoadedSuccessfully(true)
+{
+	check(LocalizationTarget);
+ 
+	const FString ManifestFile = LocalizationConfigurationScript::GetManifestPath(LocalizationTarget);
+	FString NativeCultureName;
+	if (LocalizationTarget->Settings.SupportedCulturesStatistics.IsValidIndex(LocalizationTarget->Settings.NativeCultureIndex))
+	{
+		NativeCultureName = LocalizationTarget->Settings.SupportedCulturesStatistics[LocalizationTarget->Settings.NativeCultureIndex].CultureName;
+	}
+	const FString NativeArchiveFile = NativeCultureName.IsEmpty() ? FString() : LocalizationConfigurationScript::GetArchivePath(LocalizationTarget, NativeCultureName);
+	const FString ArchiveFileToEdit = LocalizationConfigurationScript::GetArchivePath(LocalizationTarget, CultureToEdit);
+
+	OpenedManifestFilePath = ManifestFile;
+	NativeArchiveFilePath = NativeArchiveFile;
+	OpenedArchiveFilePath = ArchiveFileToEdit;
+
+	Initialize();
+}
+
+void FTranslationDataManager::Initialize()
 {
 	GWarn->BeginSlowTask(LOCTEXT("LoadingTranslationData", "Loading Translation Data..."), true);
 	TArray<UTranslationUnit*> TranslationUnits;
 
-	ManifestAtHeadRevisionPtr = ReadManifest( ManifestFilePath );
+	ManifestAtHeadRevisionPtr = ReadManifest( OpenedManifestFilePath );
 	if (ManifestAtHeadRevisionPtr.IsValid())
 	{
 		TSharedRef< FInternationalizationManifest > ManifestAtHeadRevision = ManifestAtHeadRevisionPtr.ToSharedRef();
@@ -41,7 +70,7 @@ FTranslationDataManager::FTranslationDataManager( const FString& InManifestFileP
 		{
 			bLoadedSuccessfully = false;
 			FFormatNamedArguments Arguments;
-			Arguments.Add( TEXT("ManifestFilePath"), FText::FromString(ManifestFilePath) );
+			Arguments.Add( TEXT("ManifestFilePath"), FText::FromString(OpenedManifestFilePath) );
 			Arguments.Add( TEXT("ManifestEntriesCount"), FText::AsNumber(ManifestEntriesCount) );
 			FMessageLog TranslationEditorMessageLog("TranslationEditor");
 			TranslationEditorMessageLog.Error(FText::Format(LOCTEXT("CurrentManifestEmpty", "Most current translation manifest ({ManifestFilePath}) has {ManifestEntriesCount} entries."), Arguments));
@@ -49,7 +78,9 @@ FTranslationDataManager::FTranslationDataManager( const FString& InManifestFileP
 			TranslationEditorMessageLog.Open(EMessageSeverity::Error);
 		}
 
-		ArchivePtr = ReadArchive();
+		ArchivePtr = ReadArchive(OpenedArchiveFilePath);
+		NativeArchivePtr = NativeArchiveFilePath != OpenedArchiveFilePath ? ReadArchive(NativeArchiveFilePath) : ArchivePtr;
+
 		if (ArchivePtr.IsValid())
 		{
 			int32 NumManifestEntriesParsed = 0;
@@ -76,6 +107,16 @@ FTranslationDataManager::FTranslationDataManager( const FString& InManifestFileP
 					ContextInfo.Context = AContext.SourceLocation;
 					ContextInfo.Key = AContext.Key;
 
+					if (NativeArchivePtr.IsValid() && NativeArchivePtr != ArchivePtr)
+					{
+						const TSharedPtr<FArchiveEntry> NativeArchiveEntry = NativeArchivePtr->FindEntryBySource(ManifestEntry->Namespace, ManifestEntry->Source, AContext.KeyMetadataObj);
+						// If the native archive contains a translation for the source string that isn't identical to the source string, use the translation as the source string.
+						if (NativeArchiveEntry.IsValid() && !NativeArchiveEntry->Translation.IsExactMatch(NativeArchiveEntry->Source))
+						{
+							TranslationUnit->Source = NativeArchiveEntry->Translation.Text;
+						}
+					}
+
 					TranslationUnit->Contexts.Add(ContextInfo);
 				}
 
@@ -89,7 +130,7 @@ FTranslationDataManager::FTranslationDataManager( const FString& InManifestFileP
 		{
 			bLoadedSuccessfully = false;
 			FFormatNamedArguments Arguments;
-			Arguments.Add( TEXT("ArchiveFilePath"), FText::FromString(ArchiveFilePath) );
+			Arguments.Add( TEXT("ArchiveFilePath"), FText::FromString(OpenedArchiveFilePath) );
 			FMessageLog TranslationEditorMessageLog("TranslationEditor");
 			TranslationEditorMessageLog.Error(FText::Format(LOCTEXT("FailedToLoadCurrentArchive", "Failed to load most current translation archive ({ArchiveFilePath}), unable to load translations."), Arguments));
 			TranslationEditorMessageLog.Notify(LOCTEXT("TranslationLoadError", "Error Loading Translations!"));
@@ -100,7 +141,7 @@ FTranslationDataManager::FTranslationDataManager( const FString& InManifestFileP
 	{
 		bLoadedSuccessfully = false;
 		FFormatNamedArguments Arguments;
-		Arguments.Add( TEXT("ManifestFilePath"), FText::FromString(ManifestFilePath) );
+		Arguments.Add( TEXT("ManifestFilePath"), FText::FromString(OpenedManifestFilePath) );
 		FMessageLog TranslationEditorMessageLog("TranslationEditor");
 		TranslationEditorMessageLog.Error(FText::Format(LOCTEXT("FailedToLoadCurrentManifest", "Failed to load most current translation manifest ({ManifestFilePath}), unable to load translations."), Arguments));
 		TranslationEditorMessageLog.Notify(LOCTEXT("TranslationLoadError", "Error Loading Translations!"));
@@ -133,7 +174,7 @@ TSharedPtr< FInternationalizationManifest > FTranslationDataManager::ReadManifes
 	return InternationalizationManifest;
 }
 
-TSharedPtr< FInternationalizationArchive > FTranslationDataManager::ReadArchive()
+TSharedPtr< FInternationalizationArchive > FTranslationDataManager::ReadArchive(const FString& ArchiveFilePath)
 {
 	// Read in any existing archive for this culture.
 	TSharedPtr< FJsonObject > ArchiveJsonObject = ReadJSONTextFile( ArchiveFilePath );
@@ -238,7 +279,7 @@ bool FTranslationDataManager::WriteTranslationData(bool bForceWrite /*= false*/)
 			TSharedRef<FJsonObject> FinalArchiveJsonObj = MakeShareable(new FJsonObject);
 			ArchiveSerializer.SerializeArchive(Archive, FinalArchiveJsonObj);
 
-			bSuccess = WriteJSONToTextFile(FinalArchiveJsonObj, ArchiveFilePath);
+			bSuccess = WriteJSONToTextFile(FinalArchiveJsonObj, OpenedArchiveFilePath);
 		}
 	}
 
@@ -362,7 +403,7 @@ void FTranslationDataManager::GetHistoryForTranslationUnits()
 	GWarn->BeginSlowTask(LOCTEXT("LoadingSourceControlHistory", "Loading Translation History from Source Control..."), true);
 
 	TArray<UTranslationUnit*>& TranslationUnits = AllTranslations;
-	const FString& InManifestFilePath = ManifestFilePath;
+	const FString& InManifestFilePath = OpenedManifestFilePath;
 
 	// Unload any previous history information, going to retrieve it all again.
 	UnloadHistoryInformation();
@@ -561,7 +602,7 @@ void FTranslationDataManager::HandlePropertyChanged(FName PropertyName)
 
 void FTranslationDataManager::PreviewAllTranslationsInEditor()
 {
-	FString ManifestFullPath = FPaths::ConvertRelativePathToFull(ManifestFilePath);
+	FString ManifestFullPath = FPaths::ConvertRelativePathToFull(OpenedManifestFilePath);
 	FString EngineFullPath = FPaths::ConvertRelativePathToFull(FPaths::EngineContentDir());
 
 	bool IsEngineManifest = false;
@@ -570,11 +611,10 @@ void FTranslationDataManager::PreviewAllTranslationsInEditor()
 		IsEngineManifest = true;
 	}
 
-	ULocalizationTarget* Target = ILocalizationDashboardModule::Get().GetLocalizationTargetByName(FPaths::GetBaseFilename(ManifestFilePath), IsEngineManifest);
-	if (Target)
+	if (AssociatedLocalizationTarget.IsValid())
 	{
-		const FString ConfigFilePath = LocalizationConfigurationScript::GetRegenerateResourcesScriptPath(Target);
-		LocalizationConfigurationScript::GenerateRegenerateResourcesScript(Target).Write(ConfigFilePath);
+		const FString ConfigFilePath = LocalizationConfigurationScript::GetRegenerateResourcesScriptPath(AssociatedLocalizationTarget.Get());
+		LocalizationConfigurationScript::GenerateRegenerateResourcesScript(AssociatedLocalizationTarget.Get()).Write(ConfigFilePath);
 
 		FJsonInternationalizationArchiveSerializer LocalizationArchiveSerializer;
 		FJsonInternationalizationManifestSerializer LocalizationManifestSerializer;
@@ -627,36 +667,14 @@ void FTranslationDataManager::LoadFromArchive(TArray<UTranslationUnit*>& InTrans
 
 	if (bReloadFromFile)
 	{
-		ArchivePtr = ReadArchive();
+		ArchivePtr = ReadArchive(OpenedArchiveFilePath);
+		NativeArchivePtr = NativeArchiveFilePath != OpenedArchiveFilePath ? ReadArchive(NativeArchiveFilePath) : ArchivePtr;
 	}
 
 	if (ArchivePtr.IsValid())
 	{
 		TSharedRef< FInternationalizationArchive > Archive = ArchivePtr.ToSharedRef();
 
-		for (auto Iter = Archive->GetEntryIterator(); Iter; ++Iter)
-		{
-			const FString SourceText = Iter->Key;
-			const TSharedRef< FArchiveEntry >& ArchiveEntry = Iter->Value;
-
-			const bool HasTranslationUnit = InTranslationUnits.ContainsByPredicate([&](UTranslationUnit* TU) -> bool
-			{
-				return TU->Namespace == ArchiveEntry->Namespace && TU->Source == ArchiveEntry->Source.Text;
-			});
-
-			// If a translation unit doesn't exist for this namespace and source, add it.
-			if (!HasTranslationUnit)
-			{
-				UTranslationUnit* TranslationUnit = NewObject<UTranslationUnit>();
-				check(TranslationUnit != nullptr);
-				TranslationUnit->SetFlags(RF_Transactional);
-				TranslationUnit->HasBeenReviewed = false;
-				TranslationUnit->Source = ArchiveEntry->Source.Text;
-				TranslationUnit->Namespace = ArchiveEntry->Namespace;
-				InTranslationUnits.Add(TranslationUnit);
-			}
-		}
-
 		// Make a local copy of this array before we empty the arrays below (we might have been passed AllTranslations array)
 		TArray<UTranslationUnit*> TranslationUnits;
 		TranslationUnits.Append(InTranslationUnits);
@@ -756,11 +774,12 @@ void FTranslationDataManager::LoadFromArchive(TArray<UTranslationUnit*>& InTrans
 				}
 			}
 		}
+
 	}
 	else  // ArchivePtr.IsValid() is false
 	{
 		FFormatNamedArguments Arguments;
-		Arguments.Add(TEXT("ArchiveFilePath"), FText::FromString(ArchiveFilePath));
+		Arguments.Add(TEXT("ArchiveFilePath"), FText::FromString(OpenedArchiveFilePath));
 		FMessageLog TranslationEditorMessageLog("TranslationEditor");
 		TranslationEditorMessageLog.Error(FText::Format(LOCTEXT("FailedToLoadCurrentArchive", "Failed to load most current translation archive ({ArchiveFilePath}), unable to load translations."), Arguments));
 		TranslationEditorMessageLog.Notify(LOCTEXT("TranslationLoadError", "Error Loading Translations!"));
@@ -839,9 +858,14 @@ bool FTranslationDataManager::SaveSelectedTranslations(TArray<UTranslationUnit*>
 		auto Item = *TextIt;
 		FString CurrentLocResPath = Item.Key;
 		FString ManifestAndArchiveName = FPaths::GetBaseFilename(CurrentLocResPath);
+		
 		FString ArchiveFilePath = FPaths::GetPath(CurrentLocResPath);
 		FString CultureName = FPaths::GetBaseFilename(ArchiveFilePath);
 		FString ManifestPath = FPaths::GetPath(ArchiveFilePath);
+
+		FString NativeCultureName;
+
+		FString NativeArchiveFullPath = ManifestPath / NativeCultureName / ManifestAndArchiveName + ".archive";
 		FString ArchiveFullPath = ArchiveFilePath / ManifestAndArchiveName + ".archive";
 		FString ManifestFullPath = ManifestPath / ManifestAndArchiveName + ".manifest";
 		FString EngineFullPath = FPaths::ConvertRelativePathToFull(FPaths::EngineContentDir());
@@ -853,7 +877,7 @@ bool FTranslationDataManager::SaveSelectedTranslations(TArray<UTranslationUnit*>
 
 		if (FPaths::FileExists(ManifestFullPath) && FPaths::FileExists(ArchiveFullPath))
 		{
-			TSharedRef<FTranslationDataManager> DataManager = MakeShareable(new FTranslationDataManager(ManifestFullPath, ArchiveFullPath));
+			TSharedRef<FTranslationDataManager> DataManager = MakeShareable(new FTranslationDataManager(ManifestFullPath, NativeArchiveFullPath, ArchiveFullPath));
 
 			if (DataManager->GetLoadedSuccessfully())
 			{
diff --git a/Engine/Source/Editor/TranslationEditor/Private/TranslationDataManager.h b/Engine/Source/Editor/TranslationEditor/Private/TranslationDataManager.h
--- a/Engine/Source/Editor/TranslationEditor/Private/TranslationDataManager.h
+++ b/Engine/Source/Editor/TranslationEditor/Private/TranslationDataManager.h
@@ -16,7 +16,8 @@ class FTranslationDataManager : public TSharedFromThis<FTranslationDataManager>
 
 public:
 
-	FTranslationDataManager( const FString& InManifestFilePath, const FString& InArchiveFilePath);
+	FTranslationDataManager(const FString& InManifestFilePath, const FString& InNativeArchiveFilePath, const FString& InArchiveFilePath);
+	FTranslationDataManager(ULocalizationTarget* const LocalizationTarget, const FString& CulturetoEdit);
 
 	virtual ~FTranslationDataManager();
 
@@ -87,6 +88,7 @@ public:
 	}
 
 private:
+	void Initialize();
 
 	/** Read text file into a JSON file */
 	static TSharedPtr<FJsonObject> ReadJSONTextFile( const FString& InFilePath ) ;
@@ -94,7 +96,7 @@ private:
 	/** Take a path and a manifest name and return a manifest data structure */
 	TSharedPtr< FInternationalizationManifest > ReadManifest ( const FString& ManifestFilePath );
 	/** Retrieve an archive data structure from ArchiveFilePath */
-	TSharedPtr< FInternationalizationArchive > ReadArchive();
+	TSharedPtr< FInternationalizationArchive > ReadArchive(const FString& ArchiveFilePath);
 
 	/** Write JSON file to text file */
 	bool WriteJSONToTextFile( TSharedRef<FJsonObject>& Output, const FString& Filename );
@@ -115,6 +117,8 @@ private:
 	/** Serializes and deserializes Manifests */
 	FJsonInternationalizationManifestSerializer ManifestSerializer;
 
+	/** Archive for the current project and native language */
+	TSharedPtr< FInternationalizationArchive > NativeArchivePtr;
 	/** Archive for the current project and translation language */
 	TSharedPtr< FInternationalizationArchive > ArchivePtr;
 	/** Manifest for the current project */
@@ -128,14 +132,19 @@ private:
 	FString ArchiveName;
 	/** Path to the culture (language, sort of) we are targeting */
 	FString CulturePath;
-	/** Path to the Manifest File **/
-	FString ManifestFilePath;
-	/** Path to the Archive File **/
-	FString ArchiveFilePath;
+	/** Path to the manifest file **/
+	FString OpenedManifestFilePath;
+	/** Path to the native culture's archive file **/
+	FString NativeArchiveFilePath;
+	/** Path to the archive file **/
+	FString OpenedArchiveFilePath;
 
 	/** Files that are already checked out from Perforce **/
 	TArray<FString> CheckedOutFiles;
 
+	/** The localization target associated with the files being used/edited, if any. */
+	TWeakObjectPtr<ULocalizationTarget> AssociatedLocalizationTarget;
+
 	/** Whether or not we successfully loaded the .manifest and .archive */
 	bool bLoadedSuccessfully;
 };
diff --git a/Engine/Source/Editor/TranslationEditor/Private/TranslationEditor.cpp b/Engine/Source/Editor/TranslationEditor/Private/TranslationEditor.cpp
--- a/Engine/Source/Editor/TranslationEditor/Private/TranslationEditor.cpp
+++ b/Engine/Source/Editor/TranslationEditor/Private/TranslationEditor.cpp
@@ -24,6 +24,9 @@
 #include "InternationalizationExportSettings.h"
 #include "ILocalizationServiceModule.h"
 #include "ILocalizationServiceProvider.h"
+#include "LocalizationCommandletTasks.h"
+#include "SNotificationList.h"
+#include "NotificationManager.h"
 
 DEFINE_LOG_CATEGORY_STATIC(LocalizationExport, Log, All);
 
@@ -1137,9 +1140,10 @@ void FTranslationEditor::PreviewAllTranslationsInEditor_Execute()
 
 void FTranslationEditor::ImportLatestFromLocalizationService_Execute()
 {
+	check(AssociatedLocalizationTarget.IsValid());
 	ILocalizationServiceProvider& Provider = ILocalizationServiceModule::Get().GetProvider();
 	TSharedRef<FDownloadLocalizationTargetFile, ESPMode::ThreadSafe> DownloadTargetFileOp = ILocalizationServiceOperation::Create<FDownloadLocalizationTargetFile>();
-	DownloadTargetFileOp->SetInTargetGuid(LocalizationTargetGuid);
+	DownloadTargetFileOp->SetInTargetGuid(AssociatedLocalizationTarget->Settings.Guid);
 	DownloadTargetFileOp->SetInLocale(FPaths::GetBaseFilename(FPaths::GetPath(ArchiveFilePath)));
 	FString Path = FPaths::GameSavedDir() / "Temp" / "LastImportFromLocService.po";
 	FPaths::MakePathRelativeTo(Path, *FPaths::GameDir());
@@ -1152,6 +1156,7 @@ void FTranslationEditor::ImportLatestFromLocalizationService_Execute()
 
 void FTranslationEditor::DownloadLatestFromLocalizationServiceComplete(const FLocalizationServiceOperationRef& Operation, ELocalizationServiceOperationCommandResult::Type Result)
 {
+	check(AssociatedLocalizationTarget.IsValid());
 	TSharedPtr<FDownloadLocalizationTargetFile, ESPMode::ThreadSafe> DownloadLocalizationTargetOp = StaticCastSharedRef<FDownloadLocalizationTargetFile>(Operation);
 	bool bError = !(Result == ELocalizationServiceOperationCommandResult::Succeeded);
 	FText ErrorText = FText::GetEmpty();
@@ -1177,7 +1182,7 @@ void FTranslationEditor::DownloadLatestFromLocalizationServiceComplete(const FLo
 			bError = true;
 		}
 
-		if (InTargetGuid == LocalizationTargetGuid && InLocale == FPaths::GetBaseFilename(FPaths::GetPath(ArchiveFilePath)) && !InRelativeOutputFilePathAndName.IsEmpty())
+		if (InTargetGuid == AssociatedLocalizationTarget->Settings.Guid && InLocale == FPaths::GetBaseFilename(FPaths::GetPath(ArchiveFilePath)) && !InRelativeOutputFilePathAndName.IsEmpty())
 		{
 			FString AbsoluteFilePathAndName = FPaths::ConvertRelativePathToFull(FPaths::GameDir() / InRelativeOutputFilePathAndName);
 			if (FPaths::FileExists(AbsoluteFilePathAndName))
@@ -1233,23 +1238,22 @@ void FTranslationEditor::ExportToPortableObjectFormat_Execute()
 	const FString PortableObjectFileDescription = LOCTEXT("PortableObjectFileDescription", "Portable Object File").ToString();
 	const FString PortableObjectFileExtension = TEXT("*.po");
 	const FString FileTypes = FString::Printf(TEXT("%s (%s)|%s"), *PortableObjectFileDescription, *PortableObjectFileExtension, *PortableObjectFileExtension);
-	const FString DefaultFilename = FPaths::GetBaseFilename(ManifestFilePath) + "-" + FPaths::GetBaseFilename(FPaths::GetPath(ArchiveFilePath)) + ".po";
-	FString DefaultPath = FPaths::GameSavedDir();
-	if (LastExportFilePath != "")
+	const FString CultureToEdit = FPaths::GetBaseFilename(FPaths::GetPath(ArchiveFilePath));
+	FString DefaultPath = FPaths::GetPath(LocalizationConfigurationScript::GetDefaultPOPath(AssociatedLocalizationTarget.Get(), CultureToEdit));
+	if (!LastExportFilePath.IsEmpty())
 	{
 		DefaultPath = LastExportFilePath;
 	}
 	TArray<FString> SaveFilenames;
 	IDesktopPlatform* DesktopPlatform = FDesktopPlatformModule::Get();
 	bool bSelected = false;
-	bool bHadError = false;
-	
+	const TSharedPtr<SWindow>& ParentWindow = FSlateApplication::Get().FindWidgetWindow(PreviewTextBlock);
+
 	// Prompt the user for the filename
 	if (DesktopPlatform)
 	{
 		void* ParentWindowWindowHandle = NULL;
 
-		const TSharedPtr<SWindow>& ParentWindow = FSlateApplication::Get().FindWidgetWindow(PreviewTextBlock);
 		if (ParentWindow.IsValid() && ParentWindow->GetNativeWindow().IsValid())
 		{
 			ParentWindowWindowHandle = ParentWindow->GetNativeWindow()->GetOSWindowHandle();
@@ -1258,8 +1262,8 @@ void FTranslationEditor::ExportToPortableObjectFormat_Execute()
 		bSelected = DesktopPlatform->SaveFileDialog(
 			ParentWindowWindowHandle,
 			LOCTEXT("ChooseExportLocationWindowTitle", "Choose Export Location").ToString(),
-			LastExportFilePath,
-			DefaultFilename,
+			DefaultPath,
+			LocalizationConfigurationScript::GetDefaultPOFileName(AssociatedLocalizationTarget.Get()),
 			FileTypes,
 			EFileDialogFlags::None,
 			SaveFilenames
@@ -1268,119 +1272,26 @@ void FTranslationEditor::ExportToPortableObjectFormat_Execute()
 
 	if (bSelected)
 	{
-		GWarn->BeginSlowTask(LOCTEXT("ExportingInternationalization", "Exporting Internationalization Data..."), true);
+		LastExportFilePath = FPaths::GetPath(SaveFilenames[0]);
+
+		const FString CultureToEdit = FPaths::GetBaseFilename(FPaths::GetPath(ArchiveFilePath));
 
 		// Write translation data first to ensure all changes are exported
-		DataManager->WriteTranslationData();
-
-		// Source path needs to be relative to Engine or Game directory
-		FString ManifestFullPath = FPaths::ConvertRelativePathToFull(ManifestFilePath);
-		FString EngineFullPath = FPaths::ConvertRelativePathToFull(FPaths::EngineContentDir());
-		FString FolderRelativeManifestPath = ManifestFilePath;
-		bool IsEngineManifest = false;
-		if (ManifestFullPath.StartsWith(EngineFullPath))
+		if (DataManager->WriteTranslationData() && ParentWindow.IsValid() && SaveFilenames.Num() > 0)
 		{
-			IsEngineManifest = true;
-		}
-		if (IsEngineManifest)
-		{
-			FPaths::MakePathRelativeTo(FolderRelativeManifestPath, *FPaths::EngineDir());
+			LocalizationCommandletTasks::ExportCulture(ParentWindow.ToSharedRef(), AssociatedLocalizationTarget.Get(), CultureToEdit, TOptional<FString>(SaveFilenames.Top()));
 		}
 		else
 		{
-			FPaths::MakePathRelativeTo(FolderRelativeManifestPath, *FPaths::GameDir());
-		}
-
-		// Paths like "Content/Localization/Editor/" are not considered relative for some reason, so add "./" if neccessary
-		if (!FPaths::IsRelative(FolderRelativeManifestPath))
-		{
-			FolderRelativeManifestPath = FString(TEXT("./")) + FolderRelativeManifestPath;
-		}
-
-		UInternationalizationExportSettings* ExportSettings = NewObject<UInternationalizationExportSettings>();
-		ExportSettings->CulturesToGenerate.Empty();
-		ExportSettings->CulturesToGenerate.Add(FPaths::GetBaseFilename(FPaths::GetPath(ArchiveFilePath)));
-		ExportSettings->CommandletClass = "InternationalizationExport";
-		ExportSettings->SourcePath = FPaths::GetPath(FolderRelativeManifestPath);
-		ExportSettings->ManifestName = FPaths::GetBaseFilename(ManifestFilePath) + ".manifest";
-		ExportSettings->ArchiveName = FPaths::GetBaseFilename(ManifestFilePath) + ".archive";
-		ExportSettings->bExportLoc = true;
-		ExportSettings->bImportLoc = false;
-		ExportSettings->bUseCultureDirectory = false;
-
-		ExportSettings->DestinationPath = DefaultPath / DefaultFilename;
-
-		if (SaveFilenames.Num() > 0)
-		{
-			ExportSettings->DestinationPath = FPaths::GetPath(SaveFilenames[0]);
-			ExportSettings->PortableObjectName = FPaths::GetCleanFilename(SaveFilenames[0]);
-			LastExportFilePath = FPaths::GetPath(SaveFilenames[0]);
-		}
-
-		// Write these settings to a temporary config file that the Internationalization Export Commandlet will read
-		FString TempConfigFilepath = FPaths::GameSavedDir() / "Config" / "InternationalizationExport.ini";
-		ExportSettings->SaveConfig(CPF_Config, *TempConfigFilepath);
+			FNotificationInfo Info( LOCTEXT("ExportFailedError", "Translation export failed!") );
+			Info.ExpireDuration = 4.f;
 
-		// Using .ini config saving means these settings will be saved in the GetClass()->GetPathName() section
-		TArray<FString> ConfigSections;
-		ConfigSections.Add(ExportSettings->GetClass()->GetPathName());
-		FMessageLog TranslationEditorMessageLog("TranslationEditor");
-
-		for (FString& ConfigSection : ConfigSections)
-		{
-			// Spawn the LocalizationExport commandlet, and run its log output back into ours
-			FString AppURL = FPlatformProcess::ExecutableName(true);
-			FString Parameters = FString("-run=InternationalizationExport -config=") + TempConfigFilepath + " -section=" + ConfigSection;
-
-			void* WritePipe;
-			void* ReadPipe;
-			FPlatformProcess::CreatePipe(ReadPipe, WritePipe);
-			FProcHandle ProcessHandle = FPlatformProcess::CreateProc(*AppURL, *Parameters, false, false, false, NULL, 0, NULL, WritePipe);
-
-			while (FPlatformProcess::IsProcRunning(ProcessHandle))
-			{
-				FString NewLine = FPlatformProcess::ReadPipe(ReadPipe);
-				if (NewLine.Len() > 0)
-				{
-					UE_LOG(LocalizationExport, Log, TEXT("%s"), *NewLine);
-					FFormatNamedArguments Arguments;
-					Arguments.Add(TEXT("LogMessage"), FText::FromString(NewLine));
-					TranslationEditorMessageLog.Info(FText::Format(LOCTEXT("LocalizationExportLog", "Localization Export Log: {LogMessage}"), Arguments));
-				}
-
-				FPlatformProcess::Sleep(0.25);
-			}
-			FString NewLine = FPlatformProcess::ReadPipe(ReadPipe);
-			if (NewLine.Len() > 0)
+			const TSharedPtr<SNotificationItem> NotificationItem = FSlateNotificationManager::Get().AddNotification(Info);
+			if (NotificationItem.IsValid())
 			{
-				UE_LOG(LocalizationExport, Log, TEXT("%s"), *NewLine);
-				FFormatNamedArguments Arguments;
-				Arguments.Add(TEXT("LogMessage"), FText::FromString(NewLine));
-				TranslationEditorMessageLog.Info(FText::Format(LOCTEXT("LocalizationExportLog", "Localization Export Log: {LogMessage}"), Arguments));
-			}
-
-			FPlatformProcess::Sleep(0.25);
-			FPlatformProcess::ClosePipe(ReadPipe, WritePipe);
-
-			int32 ReturnCode;
-			if (!FPlatformProcess::GetProcReturnCode(ProcessHandle, &ReturnCode))
-			{
-				bHadError = true;
-			}
-			else if (ReturnCode != 0)
-			{
-				bHadError = true;
+				NotificationItem->SetCompletionState(SNotificationItem::CS_Fail);
 			}
 		}
-
-		GWarn->EndSlowTask();
-
-		if (bHadError)
-		{
-			TranslationEditorMessageLog.Error(LOCTEXT("FailedToExportLocalization", "Failed to export localization!"));
-			TranslationEditorMessageLog.Notify(LOCTEXT("FailedToExportLocalization", "Failed to export localization!"));
-			TranslationEditorMessageLog.Open(EMessageSeverity::Error);
-		}
 	}
 }
 
@@ -1389,8 +1300,9 @@ void FTranslationEditor::ImportFromPortableObjectFormat_Execute()
 	const FString PortableObjectFileDescription = LOCTEXT("PortableObjectFileDescription", "Portable Object File").ToString();
 	const FString PortableObjectFileExtension = TEXT("*.po");
 	const FString FileTypes = FString::Printf(TEXT("%s (%s)|%s"), *PortableObjectFileDescription, *PortableObjectFileExtension, *PortableObjectFileExtension);
-	FString DefaultPath = FPaths::GameSavedDir();
-	if (LastImportFilePath != "")
+	const FString CultureToEdit = FPaths::GetBaseFilename(FPaths::GetPath(ArchiveFilePath));
+	FString DefaultPath = FPaths::GetPath(LocalizationConfigurationScript::GetDefaultPOPath(AssociatedLocalizationTarget.Get(), CultureToEdit));
+	if (!LastImportFilePath.IsEmpty())
 	{
 		DefaultPath = LastImportFilePath;
 	}
@@ -1412,7 +1324,7 @@ void FTranslationEditor::ImportFromPortableObjectFormat_Execute()
 			ParentWindowWindowHandle,
 			LOCTEXT("ChooseImportLocationWindowTitle", "Choose File to Import").ToString(),
 			DefaultPath,
-			TEXT(""),
+			LocalizationConfigurationScript::GetDefaultPOFileName(AssociatedLocalizationTarget.Get()),
 			FileTypes,
 			EFileDialogFlags::None,
 			OpenFilenames
@@ -1424,118 +1336,20 @@ void FTranslationEditor::ImportFromPortableObjectFormat_Execute()
 		FString FileToImport = OpenFilenames[0];
 
 		ImportFromPoFile(FileToImport);
-
 	}
 }
 
 void FTranslationEditor::ImportFromPoFile(FString FileToImport)
 {
-	// Source path needs to be relative to Engine or Game directory
-	FString ManifestFullPath = FPaths::ConvertRelativePathToFull(ManifestFilePath);
-	FString EngineFullPath = FPaths::ConvertRelativePathToFull(FPaths::EngineContentDir());
-	FString FolderRelativeManifestPath = ManifestFilePath;
-	bool IsEngineManifest = false;
-	if (ManifestFullPath.StartsWith(EngineFullPath))
-	{
-		IsEngineManifest = true;
-	}
-	if (IsEngineManifest)
-	{
-		FPaths::MakePathRelativeTo(FolderRelativeManifestPath, *FPaths::EngineDir());
-	}
-	else
-	{
-		FPaths::MakePathRelativeTo(FolderRelativeManifestPath, *FPaths::GameDir());
-	}
-
-	// Paths like "Content/Localization/Editor/" are not considered relative for some reason, so add "./" if neccessary
-	if (!FPaths::IsRelative(FolderRelativeManifestPath))
-	{
-		FolderRelativeManifestPath = FString(TEXT("./")) + FolderRelativeManifestPath;
-	}
-
-	UInternationalizationExportSettings* ImportSettings = NewObject<UInternationalizationExportSettings>();
-	ImportSettings->CulturesToGenerate.Empty();
-	ImportSettings->CulturesToGenerate.Add(FPaths::GetBaseFilename(FPaths::GetPath(ArchiveFilePath)));
-	ImportSettings->CommandletClass = "InternationalizationExport";
-	ImportSettings->DestinationPath = FPaths::GetPath(FolderRelativeManifestPath);
-	ImportSettings->ManifestName = FPaths::GetBaseFilename(ManifestFilePath) + ".manifest";
-	ImportSettings->ArchiveName = FPaths::GetBaseFilename(ManifestFilePath) + ".archive";
-	ImportSettings->bExportLoc = false;
-	ImportSettings->bImportLoc = true;
-	ImportSettings->bUseCultureDirectory = false;
-
-	ImportSettings->SourcePath = FPaths::GetPath(FileToImport);
-	ImportSettings->PortableObjectName = FPaths::GetCleanFilename(FileToImport);
 	LastImportFilePath = FPaths::GetPath(FileToImport);
 
-	FMessageLog TranslationEditorMessageLog("TranslationEditor");
-
 	// Write translation data first to ensure all changes are exported
-	bool bHadError = !(DataManager->WriteTranslationData(true));
+	const FString CultureToEdit = FPaths::GetBaseFilename(FPaths::GetPath(ArchiveFilePath));
+	const TSharedPtr<SWindow>& ParentWindow = FSlateApplication::Get().FindWidgetWindow(PreviewTextBlock);
 
-	if (!bHadError)
+	if (DataManager->WriteTranslationData(true) && ParentWindow.IsValid())
 	{
-		GWarn->BeginSlowTask(LOCTEXT("ImportingInternationalization", "Importing Internationalization Data..."), true);
-
-		// Write these settings to a temporary config file that the Internationalization Export Commandlet will read
-		FString TempConfigFilepath = FPaths::GameSavedDir() / "Config" / "InternationalizationExport.ini";
-		ImportSettings->SaveConfig(CPF_Config, *TempConfigFilepath);
-
-		// Using .ini config saving means these settings will be saved in the GetClass()->GetPathName() section
-		TArray<FString> ConfigSections;
-		ConfigSections.Add(ImportSettings->GetClass()->GetPathName());
-
-		for (FString& ConfigSection : ConfigSections)
-		{
-			// Spawn the LocalizationExport commandlet, and run its log output back into ours
-			FString AppURL = FPlatformProcess::ExecutableName(true);
-			FString Parameters = FString("-run=InternationalizationExport -config=") + TempConfigFilepath + " -section=" + ConfigSection;
-
-			void* WritePipe;
-			void* ReadPipe;
-			FPlatformProcess::CreatePipe(ReadPipe, WritePipe);
-			FProcHandle ProcessHandle = FPlatformProcess::CreateProc(*AppURL, *Parameters, false, false, false, NULL, 0, NULL, WritePipe);
-
-			while (FPlatformProcess::IsProcRunning(ProcessHandle))
-			{
-				FString NewLine = FPlatformProcess::ReadPipe(ReadPipe);
-				if (NewLine.Len() > 0)
-				{
-					UE_LOG(LocalizationExport, Log, TEXT("%s"), *NewLine);
-					FFormatNamedArguments Arguments;
-					Arguments.Add(TEXT("LogMessage"), FText::FromString(NewLine));
-					TranslationEditorMessageLog.Info(FText::Format(LOCTEXT("LocalizationImportLog", "Localization Import Log: {LogMessage}"), Arguments));
-				}
-
-				FPlatformProcess::Sleep(0.25);
-			}
-			FString NewLine = FPlatformProcess::ReadPipe(ReadPipe);
-			if (NewLine.Len() > 0)
-			{
-				UE_LOG(LocalizationExport, Log, TEXT("%s"), *NewLine);
-				FFormatNamedArguments Arguments;
-				Arguments.Add(TEXT("LogMessage"), FText::FromString(NewLine));
-				TranslationEditorMessageLog.Info(FText::Format(LOCTEXT("LocalizationImportLog", "Localization Import Log: {LogMessage}"), Arguments));
-			}
-
-			FPlatformProcess::Sleep(0.25);
-			FPlatformProcess::ClosePipe(ReadPipe, WritePipe);
-
-			int32 ReturnCode;
-			if (!FPlatformProcess::GetProcReturnCode(ProcessHandle, &ReturnCode))
-			{
-				bHadError = true;
-			}
-			else if (ReturnCode != 0)
-			{
-				bHadError = true;
-			}
-		}
-
-		GWarn->EndSlowTask();
-
-		if (!bHadError)
+		if (LocalizationCommandletTasks::ImportCulture(ParentWindow.ToSharedRef(), AssociatedLocalizationTarget.Get(), CultureToEdit, TOptional<FString>(FileToImport)))
 		{
 			DataManager->LoadFromArchive(DataManager->GetAllTranslationsArray(), true, true);
 
@@ -1547,13 +1361,16 @@ void FTranslationEditor::ImportFromPoFile(FString FileToImport)
 			ChangedOnImportPropertyTable->AddColumn((TWeakObjectPtr<UProperty>)FindField<UProperty>(UTranslationUnit::StaticClass(), "Translation"));
 		}
 	}
-
-	if (bHadError)
+	else
 	{
-		FText FailedToImport = LOCTEXT("FailedToImportLocalization", "Failed to import localization data!");
-		TranslationEditorMessageLog.Error(FailedToImport);
-		TranslationEditorMessageLog.Notify(FailedToImport, EMessageSeverity::Info, true);
-		TranslationEditorMessageLog.Open(EMessageSeverity::Error);
+		FNotificationInfo Info( LOCTEXT("ImportFailedError", "Translation import failed!") );
+		Info.ExpireDuration = 4.f;
+
+		const TSharedPtr<SNotificationItem> NotificationItem = FSlateNotificationManager::Get().AddNotification(Info);
+		if (NotificationItem.IsValid())
+		{
+			NotificationItem->SetCompletionState(SNotificationItem::CS_Fail);
+		}
 	}
 }
 
diff --git a/Engine/Source/Editor/TranslationEditor/Private/TranslationEditor.h b/Engine/Source/Editor/TranslationEditor/Private/TranslationEditor.h
--- a/Engine/Source/Editor/TranslationEditor/Private/TranslationEditor.h
+++ b/Engine/Source/Editor/TranslationEditor/Private/TranslationEditor.h
@@ -7,18 +7,17 @@
 #include "CustomFontColumn.h"
 #include "TranslationUnit.h"
 #include "ILocalizationServiceProvider.h"
+#include "LocalizationConfigurationScript.h"
 
 class TRANSLATIONEDITOR_API FTranslationEditor :  public ITranslationEditor
 {
 public:
-
-	
 	/**
 	 *	Creates a new FTranslationEditor and calls Initialize
 	 */
-	static TSharedRef< FTranslationEditor > Create(TSharedRef< FTranslationDataManager > DataManager, const FString& InManifestFile, const FString& InArchiveFile, const FGuid& LocalizationTargetGuid)
+	static TSharedRef< FTranslationEditor > Create(TSharedRef< FTranslationDataManager > DataManager, const FString& InManifestFile, const FString& InArchiveFile)
 	{
-		TSharedRef< FTranslationEditor > TranslationEditor = MakeShareable(new FTranslationEditor(DataManager, InManifestFile, InArchiveFile, LocalizationTargetGuid));
+		TSharedRef< FTranslationEditor > TranslationEditor = MakeShareable(new FTranslationEditor(DataManager, InManifestFile, InArchiveFile, nullptr));
 
 		// Some stuff that needs to use the "this" pointer is done in Initialize (because it can't be done in the constructor)
 		TranslationEditor->Initialize();
@@ -32,6 +31,24 @@ public:
 		return TranslationEditor;
 	}
 	
+	static TSharedRef< FTranslationEditor > Create(TSharedRef< FTranslationDataManager > DataManager, ULocalizationTarget* const LocalizationTarget, const FString& CultureToEdit)
+	{
+		check(LocalizationTarget);
+		
+		TSharedRef< FTranslationEditor > TranslationEditor = MakeShareable(new FTranslationEditor(DataManager, LocalizationConfigurationScript::GetManifestPath(LocalizationTarget), LocalizationConfigurationScript::GetArchivePath(LocalizationTarget, CultureToEdit), LocalizationTarget));
+
+		// Some stuff that needs to use the "this" pointer is done in Initialize (because it can't be done in the constructor)
+		TranslationEditor->Initialize();
+
+		for (UTranslationUnit* TranslationUnit : DataManager->GetAllTranslationsArray())
+		{
+			// Set up a property changed event to trigger a write of the translation data when TranslationUnit property changes
+			TranslationUnit->OnPropertyChanged().AddSP(DataManager, &FTranslationDataManager::HandlePropertyChanged);
+		}
+
+		return TranslationEditor;
+	}
+
 	virtual ~FTranslationEditor() {}
 
 	virtual void RegisterTabSpawners(const TSharedRef<class FTabManager>& TabManager) override;
@@ -60,8 +77,8 @@ protected:
 
 private:
 
-	FTranslationEditor(TSharedRef< FTranslationDataManager > InDataManager, const FString& InManifestFile, const FString& InArchiveFile, const FGuid& InLocalizationTargetGuid)
-		: ITranslationEditor(InManifestFile, InArchiveFile, InLocalizationTargetGuid)
+	FTranslationEditor(TSharedRef< FTranslationDataManager > InDataManager, const FString& InManifestFile, const FString& InArchiveFile, ULocalizationTarget* const LocalizationTarget)
+		: ITranslationEditor(InManifestFile, InArchiveFile, LocalizationTarget)
 	, DataManager(InDataManager)
 	, SourceFont(FEditorStyle::GetFontStyle( PropertyTableConstants::NormalFontStyle ))
 	, TranslationTargetFont(FEditorStyle::GetFontStyle( PropertyTableConstants::NormalFontStyle ))
diff --git a/Engine/Source/Editor/TranslationEditor/Private/TranslationEditorModule.cpp b/Engine/Source/Editor/TranslationEditor/Private/TranslationEditorModule.cpp
--- a/Engine/Source/Editor/TranslationEditor/Private/TranslationEditorModule.cpp
+++ b/Engine/Source/Editor/TranslationEditor/Private/TranslationEditorModule.cpp
@@ -6,6 +6,7 @@
 //#include "Toolkits/ToolkitManager.h"
 #include "WorkspaceMenuStructureModule.h"
 #include "MessageLogModule.h"
+#include "LocalizationConfigurationScript.h"
 
 class FTranslationEditor;
 
@@ -39,14 +40,14 @@ void FTranslationEditorModule::ShutdownModule()
 #endif
 }
 
-TSharedRef<FTranslationEditor> FTranslationEditorModule::CreateTranslationEditor(const FString& ManifestFile, const FString& ArchiveFile, const FGuid& LocalizationTargetGuid, bool& OutLoadedSuccessfully)
+TSharedRef<FTranslationEditor> FTranslationEditorModule::CreateTranslationEditor(const FString& ManifestFile, const FString& NativeArchiveFile, const FString& ArchiveFileToEdit, bool& OutLoadedSuccessfully)
 {
-	TSharedRef< FTranslationDataManager > DataManager = MakeShareable( new FTranslationDataManager(ManifestFile, ArchiveFile) );
+	TSharedRef< FTranslationDataManager > DataManager = MakeShareable( new FTranslationDataManager(ManifestFile, NativeArchiveFile, ArchiveFileToEdit) );
 	OutLoadedSuccessfully = DataManager->GetLoadedSuccessfully();
 
 	GWarn->BeginSlowTask(LOCTEXT("BuildingUserInterface", "Building Translation Editor UI..."), true);
 
-	TSharedRef< FTranslationEditor > NewTranslationEditor(FTranslationEditor::Create(DataManager, ManifestFile, ArchiveFile, LocalizationTargetGuid));
+	TSharedRef< FTranslationEditor > NewTranslationEditor(FTranslationEditor::Create(DataManager, ManifestFile, ArchiveFileToEdit));
 	NewTranslationEditor->InitTranslationEditor(EToolkitMode::Standalone, TSharedPtr<IToolkitHost>());
 
 	GWarn->EndSlowTask();
@@ -54,4 +55,29 @@ TSharedRef<FTranslationEditor> FTranslationEditorModule::CreateTranslationEditor
 	return NewTranslationEditor;
 }
 
+TSharedRef<FTranslationEditor> FTranslationEditorModule::CreateTranslationEditor(ULocalizationTarget* const LocalizationTarget, const FString& CultureToEdit, bool& OutLoadedSuccessfully)
+{
+	const FString ManifestFile = LocalizationConfigurationScript::GetManifestPath(LocalizationTarget);
+	FString NativeCultureName;
+	if (LocalizationTarget->Settings.SupportedCulturesStatistics.IsValidIndex(LocalizationTarget->Settings.NativeCultureIndex))
+	{
+		NativeCultureName = LocalizationTarget->Settings.SupportedCulturesStatistics[LocalizationTarget->Settings.NativeCultureIndex].CultureName;
+	}
+	const FString NativeArchiveFile = NativeCultureName.IsEmpty() ? FString() : LocalizationConfigurationScript::GetArchivePath(LocalizationTarget, NativeCultureName);
+	const FString ArchiveFileToEdit = LocalizationConfigurationScript::GetArchivePath(LocalizationTarget, CultureToEdit);
+
+	TSharedRef< FTranslationDataManager > DataManager = MakeShareable( new FTranslationDataManager(LocalizationTarget, CultureToEdit) );
+	OutLoadedSuccessfully = DataManager->GetLoadedSuccessfully();
+
+	GWarn->BeginSlowTask(LOCTEXT("BuildingUserInterface", "Building Translation Editor UI..."), true);
+
+	TSharedRef< FTranslationEditor > NewTranslationEditor(FTranslationEditor::Create(DataManager, LocalizationTarget, CultureToEdit));
+	NewTranslationEditor->InitTranslationEditor(EToolkitMode::Standalone, TSharedPtr<IToolkitHost>());
+
+	GWarn->EndSlowTask();
+
+	return NewTranslationEditor;
+}
+
+
 #undef LOCTEXT_NAMESPACE
\ No newline at end of file
diff --git a/Engine/Source/Editor/TranslationEditor/Public/ITranslationEditor.h b/Engine/Source/Editor/TranslationEditor/Public/ITranslationEditor.h
--- a/Engine/Source/Editor/TranslationEditor/Public/ITranslationEditor.h
+++ b/Engine/Source/Editor/TranslationEditor/Public/ITranslationEditor.h
@@ -5,35 +5,34 @@
 #include "Toolkits/IToolkit.h"
 #include "Toolkits/AssetEditorToolkit.h"
 #include "TranslationEditorModule.h"
-#include "../Private/TranslationPickerWidget.h"
-
+#include "Private/TranslationPickerWidget.h"
+#include "LocalizationTargetTypes.h"
 
 /** Translation Editor public interface */
 class ITranslationEditor : public FAssetEditorToolkit
 {
-
 public:
-
-	TRANSLATIONEDITOR_API static void OpenTranslationEditor(const FString& InManifestFile, const FString& InArchiveFile, const FGuid& LocalizationTargetGuid);
+	TRANSLATIONEDITOR_API static void OpenTranslationEditor(const FString& InManifestFile, const FString& InNativeArchiveFile, const FString& InArchiveFileToEdit);
+	TRANSLATIONEDITOR_API static void OpenTranslationEditor(ULocalizationTarget* const LocalizationTarget, const FString& CultureToEdit);
 
 	TRANSLATIONEDITOR_API static void OpenTranslationPicker();
 
-	ITranslationEditor(const FString& InManifestFile, const FString& InArchiveFile, const FGuid& InLocalizationTargetGuid)
+	ITranslationEditor(const FString& InManifestFile, const FString& InArchiveFile, ULocalizationTarget* const InAssociatedLocalizationTarget)
 		: ManifestFilePath(InManifestFile)
 		, ArchiveFilePath(InArchiveFile)
-		, LocalizationTargetGuid(InLocalizationTargetGuid)
+		, AssociatedLocalizationTarget(InAssociatedLocalizationTarget)
 	{}
 
 	virtual bool OnRequestClose() override;
 
 protected:
 
-	/** Name of the project we are translating for */
+	/** The path to the manifest file being used for contexts. */
 	FString ManifestFilePath;
-	/** Name of the language we are translating to */
+	/** The path to the archive file being edited. */
 	FString ArchiveFilePath;
-	/** GUID of the Localization Target we are editing */
-	FGuid LocalizationTargetGuid;
+	/** The localization target associated with the files being used/edited, if any. */
+	TWeakObjectPtr<ULocalizationTarget> AssociatedLocalizationTarget;
 
 private:
 
diff --git a/Engine/Source/Editor/TranslationEditor/Public/TranslationEditorModule.h b/Engine/Source/Editor/TranslationEditor/Public/TranslationEditorModule.h
--- a/Engine/Source/Editor/TranslationEditor/Public/TranslationEditorModule.h
+++ b/Engine/Source/Editor/TranslationEditor/Public/TranslationEditorModule.h
@@ -8,6 +8,7 @@
 //#include "ITranslationEditor.h"
 
 class FTranslationEditor;
+class ULocalizationTarget;
 
 class FTranslationEditorModule : public IModuleInterface,
 	public IHasMenuExtensibility
@@ -21,14 +22,26 @@ public:
 	/**
 	 * Creates an instance of translation editor object.  Only virtual so that it can be called across the DLL boundary.
 	 *
-	 * @param ProjectName				Name of the project to translate
-	 * @param TranslationTargetLanguage	Language to translate to
-	 * @param LocalizationTargetGuid	Unique ID for this Localization Target
-	 * @param OutLoadedSuccessfully		Whether or not the translation editor was able to load successfully from the .manifest and .archive files
+	 * @param ManifestFile			The path to the manifest file to be used for contexts.
+	 * @param NativeArchiveFile		The path to the archive file for the native language.
+	 * @param ArchiveFileToEdit		The path to the archive file to be viewed and edited.
+	 * @param OutLoadedSuccessfully	Whether or not the translation editor was able to load successfully from the .manifest and .archive files
 	 * 
-	 * @return	Interface to the new translation editor
+	 * @return	The new instance of the translation editor
 	 */
-	virtual TSharedRef<FTranslationEditor> CreateTranslationEditor(const FString& ManifestFile, const FString& ArchiveFile, const FGuid& LocalizationTargetGuid, bool& OutLoadedSuccessfully);
+	virtual TSharedRef<FTranslationEditor> CreateTranslationEditor(const FString& ManifestFile, const FString& NativeArchiveFile, const FString& ArchiveFileToEdit, bool& OutLoadedSuccessfully);
+
+	/**
+	 * Creates an instance of translation editor object.  Only virtual so that it can be called across the DLL boundary.
+	 *
+	 * @param LocalizationTarget	The localization target whose data is to be used and edited.
+	 * @param CultureToEdit			The name of a supported culture of the localization target, whose archives should be edited.
+	 * @param OutLoadedSuccessfully	Whether or not the translation editor was able to load successfully from the .manifest and .archive files
+	 * 
+	 * @return	The new instance of the translation editor
+	 */
+	virtual TSharedRef<FTranslationEditor> CreateTranslationEditor(ULocalizationTarget* const LocalizationTarget, const FString& CultureToEdit, bool& OutLoadedSuccessfully);
+
 
 	/** Gets the extensibility managers for outside entities to extend translation editor's menus and toolbars */
 	virtual TSharedPtr<FExtensibilityManager> GetMenuExtensibilityManager() override { return MenuExtensibilityManager; }
diff --git a/Engine/Source/Editor/TranslationEditor/TranslationEditor.Build.cs b/Engine/Source/Editor/TranslationEditor/TranslationEditor.Build.cs
--- a/Engine/Source/Editor/TranslationEditor/TranslationEditor.Build.cs
+++ b/Engine/Source/Editor/TranslationEditor/TranslationEditor.Build.cs
@@ -46,6 +46,7 @@ public class TranslationEditor : ModuleRules
                 "Core",
 				"CoreUObject",
 				"Engine",
+                "Localization",
             }
         );
 
@@ -55,5 +56,19 @@ public class TranslationEditor : ModuleRules
 				"DesktopPlatform",
 			}
 		);
+
+        PrivateIncludePaths.AddRange(
+            new string[]
+			{
+				"Editor/TranslationEditor/Private",
+			}
+        );
+
+        PublicIncludePaths.AddRange(
+            new string[]
+			{
+				"Editor/TranslationEditor/Public",
+			}
+        );
 	}
 }
