Added ShouldGatherFromEditorOnlyData setting to localization dashboard UI for gathering from assets and metadata.

diff --git a/Engine/Source/Editor/Localization/Private/LocalizationConfigurationScript.cpp b/Engine/Source/Editor/Localization/Private/LocalizationConfigurationScript.cpp
--- a/Engine/Source/Editor/Localization/Private/LocalizationConfigurationScript.cpp
+++ b/Engine/Source/Editor/Localization/Private/LocalizationConfigurationScript.cpp
@@ -198,6 +198,8 @@ namespace LocalizationConfigurationScript
 			{
 				ConfigSection.Add( TEXT("PackageFileNameFilters"), FString::Printf( TEXT("*.%s"), *FileExtension.Pattern) );
 			}
+
+			ConfigSection.Add( TEXT("ShouldGatherFromEditorOnlyData"), Target->Settings.GatherFromPackages.ShouldGatherFromEditorOnlyData ? TEXT("true") : TEXT("false") );
 		}
 
 		// GatherTextFromMetadata
@@ -227,6 +229,8 @@ namespace LocalizationConfigurationScript
 				ConfigSection.Add( TEXT("OutputNamespaces"), Specification.TextNamespace );
 				ConfigSection.Add( TEXT("OutputKeys"), FString::Printf(TEXT("\"%s\""), *Specification.TextKeyPattern.Pattern) );
 			}
+
+			ConfigSection.Add( TEXT("ShouldGatherFromEditorOnlyData"), Target->Settings.GatherFromMetaData.ShouldGatherFromEditorOnlyData ? TEXT("true") : TEXT("false") );
 		}
 
 		// GenerateGatherManifest
diff --git a/Engine/Source/Editor/Localization/Public/LocalizationTargetTypes.h b/Engine/Source/Editor/Localization/Public/LocalizationTargetTypes.h
--- a/Engine/Source/Editor/Localization/Public/LocalizationTargetTypes.h
+++ b/Engine/Source/Editor/Localization/Public/LocalizationTargetTypes.h
@@ -103,6 +103,7 @@ struct FGatherTextFromPackagesConfiguration
 	FGatherTextFromPackagesConfiguration()
 		: IsEnabled(true)
 		, FileExtensions(GetDefaultPackageFileExtensions())
+		, ShouldGatherFromEditorOnlyData(false)
 	{
 	}
 
@@ -122,6 +123,10 @@ struct FGatherTextFromPackagesConfiguration
 	UPROPERTY(config, EditAnywhere, Category = "Filter")
 	TArray<FGatherTextFileExtension> FileExtensions;
 
+	/* If enable, data that is specified as editor-only may be processed for gathering. */
+	UPROPERTY(config, EditAnywhere, Category = "Filter")
+	bool ShouldGatherFromEditorOnlyData;
+
 	LOCALIZATION_API bool Validate(const FString& RootDirectory, FText& OutError) const;
 };
 
@@ -179,6 +184,7 @@ struct FGatherTextFromMetaDataConfiguration
 
 	FGatherTextFromMetaDataConfiguration()
 		: IsEnabled(false)
+		, ShouldGatherFromEditorOnlyData(false)
 	{
 	}
 
@@ -198,6 +204,10 @@ struct FGatherTextFromMetaDataConfiguration
 	UPROPERTY(config, EditAnywhere, Category = "MetaData")
 	TArray<FMetaDataKeyGatherSpecification> KeySpecifications;
 
+	/* If enable, data that is specified as editor-only may be processed for gathering. */
+	UPROPERTY(config, EditAnywhere, Category = "Filter")
+	bool ShouldGatherFromEditorOnlyData;
+
 	LOCALIZATION_API bool Validate(const FString& RootDirectory, FText& OutError) const;
 };
 
