VS2015 code rework for ICE

diff --git a/Engine/Source/Editor/Localization/Private/LocalizationTargetTypes.cpp b/Engine/Source/Editor/Localization/Private/LocalizationTargetTypes.cpp
--- a/Engine/Source/Editor/Localization/Private/LocalizationTargetTypes.cpp
+++ b/Engine/Source/Editor/Localization/Private/LocalizationTargetTypes.cpp
@@ -300,7 +300,7 @@ namespace
 				}
 				for (int32 i = 0; i < MandatoryColumnNames.Num(); ++i)
 				{
-					if (MandatoryColumnNames[i] != HeaderRow[i])
+					if (!MandatoryColumnNames[i].Equals(HeaderRow[i]))
 					{
 						return false;
 					}
